Fixed up redundant FText .ToString() usage and some improper FText usage.

diff --git a/Engine/Source/Editor/LevelEditor/Private/SLevelEditorBuildAndSubmit.cpp b/Engine/Source/Editor/LevelEditor/Private/SLevelEditorBuildAndSubmit.cpp
--- a/Engine/Source/Editor/LevelEditor/Private/SLevelEditorBuildAndSubmit.cpp
+++ b/Engine/Source/Editor/LevelEditor/Private/SLevelEditorBuildAndSubmit.cpp
@@ -61,22 +61,22 @@ public:
 		{
 			// The state column shows the source control status. It should only ever be one of the three states shown below.
 			// Uses a smaller font and uses CAPS to make them stand out next to the name column.
-			FString StateText;
+			FText StateText;
 			if(Item->SourceControlState->IsCheckedOut())
 			{
-				StateText = LOCTEXT("SourceControlState_CheckedOut", "CHECKED OUT").ToString();
+				StateText = LOCTEXT("SourceControlState_CheckedOut", "CHECKED OUT");
 			}
 			else if(!Item->SourceControlState->IsSourceControlled())
 			{
-				StateText = LOCTEXT("SourceControlState_NotInDepot", "NOT IN DEPOT").ToString();
+				StateText = LOCTEXT("SourceControlState_NotInDepot", "NOT IN DEPOT");
 			}
 			else if(Item->SourceControlState->IsAdded())
 			{
-				StateText = LOCTEXT("SourceControlState_OpenForAdd", "OPEN FOR ADD").ToString();
+				StateText = LOCTEXT("SourceControlState_OpenForAdd", "OPEN FOR ADD");
 			}
 			else
 			{
-				StateText = LOCTEXT("SourceControlState_Unknown", "UNKNOWN").ToString();
+				StateText = LOCTEXT("SourceControlState_Unknown", "UNKNOWN");
 			}
 
 			return SNew(SHorizontalBox)
@@ -142,7 +142,7 @@ void SLevelEditorBuildAndSubmit::Construct( const FArguments& InArgs, const TSha
 			.Padding( 0.0f, 3.0f, 0.0f, 0.0f )
 			[
 				SNew( STextBlock )
-				.Text(LOCTEXT("ColumnHeader_Name", "Name").ToString())
+				.Text(LOCTEXT("ColumnHeader_Name", "Name"))
 			]
 		]
 		+ SHeaderRow::Column(NAME_StateColumn)
@@ -154,7 +154,7 @@ void SLevelEditorBuildAndSubmit::Construct( const FArguments& InArgs, const TSha
 			.Padding( 0.0f, 3.0f, 0.0f, 0.0f )
 			[
 				SNew( STextBlock )
-				.Text(LOCTEXT("ColumnHeader_State", "State").ToString())
+				.Text(LOCTEXT("ColumnHeader_State", "State"))
 			]
 		];
 
@@ -166,7 +166,7 @@ void SLevelEditorBuildAndSubmit::Construct( const FArguments& InArgs, const TSha
 		.Padding(5)
 		[
 			SNew(SExpandableArea)
-			.AreaTitle(LOCTEXT("DescriptionSectionTitle", "Submission Description").ToString())
+			.AreaTitle(LOCTEXT("DescriptionSectionTitle", "Submission Description"))
 			.Padding(2)
 			.BodyContent()
 			[
@@ -179,7 +179,7 @@ void SLevelEditorBuildAndSubmit::Construct( const FArguments& InArgs, const TSha
 		.Padding(5)
 		[
 			SNew(SExpandableArea )
-			.AreaTitle(LOCTEXT("ExtraPackagesSectionTitle", "Additional Files to Submit").ToString())
+			.AreaTitle(LOCTEXT("ExtraPackagesSectionTitle", "Additional Files to Submit"))
 			.Padding(2)
 			.InitiallyCollapsed(true)
 			.MaxHeight(300)
@@ -198,7 +198,7 @@ void SLevelEditorBuildAndSubmit::Construct( const FArguments& InArgs, const TSha
 		.Padding(5)
 		[
 			SNew( SExpandableArea )
-			.AreaTitle(LOCTEXT("BuildOptionsSectionTitle", "Build Options").ToString())
+			.AreaTitle(LOCTEXT("BuildOptionsSectionTitle", "Build Options"))
 			.Padding(2)
 			.BodyContent()
 			[
@@ -210,7 +210,7 @@ void SLevelEditorBuildAndSubmit::Construct( const FArguments& InArgs, const TSha
 					SAssignNew(NoSubmitOnMapErrorBox, SCheckBox)
 					[
 						SNew(STextBlock)
-						.Text(LOCTEXT("NoSubmitMapErrorsButtonLabel", "Don't Submit in Event of Map Errors").ToString())
+						.Text(LOCTEXT("NoSubmitMapErrorsButtonLabel", "Don't Submit in Event of Map Errors"))
 					]
 				]
 				+SVerticalBox::Slot()
@@ -220,7 +220,7 @@ void SLevelEditorBuildAndSubmit::Construct( const FArguments& InArgs, const TSha
 					SAssignNew(NoSubmitOnSaveErrorBox, SCheckBox)
 					[
 						SNew(STextBlock)
-						.Text(LOCTEXT("NoSubmitSaveErrorsButtonLabel", "Don't Submit in Event of Save Errors").ToString())
+						.Text(LOCTEXT("NoSubmitSaveErrorsButtonLabel", "Don't Submit in Event of Save Errors"))
 					]
 				]
 				+SVerticalBox::Slot()
@@ -231,7 +231,7 @@ void SLevelEditorBuildAndSubmit::Construct( const FArguments& InArgs, const TSha
 					.OnCheckStateChanged(this, &SLevelEditorBuildAndSubmit::OnShowPackagesNotInSCBoxChanged)
 					[
 						SNew(STextBlock)
-						.Text(LOCTEXT("ShowPackagesButtonLabel", "Show Files not in Source Control").ToString())
+						.Text(LOCTEXT("ShowPackagesButtonLabel", "Show Files not in Source Control"))
 					]
 				]
 				+SVerticalBox::Slot()
@@ -242,7 +242,7 @@ void SLevelEditorBuildAndSubmit::Construct( const FArguments& InArgs, const TSha
 					.IsChecked(ESlateCheckBoxState::Checked)
 					[
 						SNew(STextBlock)
-						.Text(LOCTEXT("AddFilesButtonLabel", "Add Files to Source Control if Necessary").ToString())
+						.Text(LOCTEXT("AddFilesButtonLabel", "Add Files to Source Control if Necessary"))
 					]
 				]
 			]			
@@ -259,14 +259,14 @@ void SLevelEditorBuildAndSubmit::Construct( const FArguments& InArgs, const TSha
 			[
 				SNew(SButton)
 				.OnClicked(this, &SLevelEditorBuildAndSubmit::OnBuildAndCloseClicked)
-				.Text(LOCTEXT("BuildAndCloseButtonLabel", "Build and Close").ToString())
+				.Text(LOCTEXT("BuildAndCloseButtonLabel", "Build and Close"))
 			]
 			+SHorizontalBox::Slot()
 			.AutoWidth()
 			[
 				SNew(SButton)
 				.OnClicked(this, &SLevelEditorBuildAndSubmit::OnBuildClicked)
-				.Text(LOCTEXT("BuildButtonLabel", "Build").ToString())
+				.Text(LOCTEXT("BuildButtonLabel", "Build"))
 			]
 		]
 	];
diff --git a/Engine/Source/Editor/LevelEditor/Private/SSurfaceProperties.cpp b/Engine/Source/Editor/LevelEditor/Private/SSurfaceProperties.cpp
--- a/Engine/Source/Editor/LevelEditor/Private/SSurfaceProperties.cpp
+++ b/Engine/Source/Editor/LevelEditor/Private/SSurfaceProperties.cpp
@@ -117,7 +117,7 @@ TSharedRef<SWidget> SSurfaceProperties::ConstructTexturePan()
 	.Padding(0,0,0,5)
 	[
 		SNew(STextBlock)
-		.Text( LOCTEXT("Pan", "Pan:").ToString() )
+		.Text( LOCTEXT("Pan", "Pan:") )
 	];
 
 	TSharedPtr<SHorizontalBox> HorizontalBox;
@@ -148,7 +148,7 @@ TSharedRef<SWidget> SSurfaceProperties::ConstructTexturePan()
 			.IsChecked( this, &SSurfaceProperties::IsUsingNegativePanning, Channels[i] )
 			.OnCheckStateChanged( this, &SSurfaceProperties::OnTogglePanningDirection, Channels[i])
 			.Style( FEditorStyle::Get(), "TransparentCheckBox" )
-			.ToolTipText( LOCTEXT("InvertPanningDirection", "Toggle panning direction.").ToString() )
+			.ToolTipText( LOCTEXT("InvertPanningDirection", "Toggle panning direction.") )
 			[
 				SNew( SImage )
 				.Image( this, &SSurfaceProperties::GetTogglePanDirectionImage, Channels[i] )
@@ -176,7 +176,7 @@ TSharedRef<SWidget> SSurfaceProperties::ConstructTexturePan()
 			.HAlign(HAlign_Center)
 			.VAlign(VAlign_Center)
 			.Text(ButtonFields[i])
-			.ToolTipText(LOCTEXT("PanUTooltip", "Pans U texture coordinate").ToString())
+			.ToolTipText(LOCTEXT("PanUTooltip", "Pans U texture coordinate"))
 			.OnClicked(this,&SSurfaceProperties::OnPanTexture,ButtonIncriments[i],SSurfaceProperties::UChannel)
 		];
 	
@@ -188,7 +188,7 @@ TSharedRef<SWidget> SSurfaceProperties::ConstructTexturePan()
 			.HAlign(HAlign_Center)
 			.VAlign(VAlign_Center)
 			.Text(ButtonFields[i])
-			.ToolTipText(LOCTEXT("PanVTooltip", "Pans V texture coordinate").ToString())
+			.ToolTipText(LOCTEXT("PanVTooltip", "Pans V texture coordinate"))
 			.OnClicked(this,&SSurfaceProperties::OnPanTexture,ButtonIncriments[i],SSurfaceProperties::VChannel)
 		];
 	}
@@ -212,8 +212,8 @@ TSharedRef<SWidget> SSurfaceProperties::ConstructTexturePan()
 				.HAlign(HAlign_Center)
 				[
 					SNew(STextBlock)
-					.Text(LOCTEXT("PanToolCustomPan", "---").ToString())
-					.ToolTipText(LOCTEXT("PanToolCustomPanToolTip", "Set Custom pan amount").ToString())
+					.Text(LOCTEXT("PanToolCustomPan", "---"))
+					.ToolTipText(LOCTEXT("PanToolCustomPanToolTip", "Set Custom pan amount"))
 				]
 			] 
 			.MenuContent()
@@ -247,7 +247,7 @@ TSharedRef<SWidget> SSurfaceProperties::ConstructTextureRotate()
 	.Padding(0,0,5,5)
 	[
 		SNew (STextBlock)
-		.Text(LOCTEXT("RotateTitle", "Rotate:").ToString())
+		.Text(LOCTEXT("RotateTitle", "Rotate:"))
 	];
 
 	TSharedPtr<SHorizontalBox> RotateBox;
@@ -267,7 +267,7 @@ TSharedRef<SWidget> SSurfaceProperties::ConstructTextureRotate()
 		.IsChecked( this, &SSurfaceProperties::IsUsingNegativeRotation )
 		.OnCheckStateChanged( this, &SSurfaceProperties::OnToggleRotationDirection )
 		.Style( FEditorStyle::Get(), "TransparentCheckBox" )
-		.ToolTipText( LOCTEXT("InvertRotation", "Toggle Rotation direction.").ToString() )
+		.ToolTipText( LOCTEXT("InvertRotation", "Toggle Rotation direction.") )
 		[
 			SNew( SImage )
 			.Image( this, &SSurfaceProperties::GetToggleRotationDirectionImage )
@@ -276,7 +276,7 @@ TSharedRef<SWidget> SSurfaceProperties::ConstructTextureRotate()
 	];
 
 	// Rotation button fields:
-	const static FString ButtonFields[] = {"45", "90", LOCTEXT("RotateCustom", "---").ToString()};
+	const static FText ButtonFields[] = { LOCTEXT("Rotate45Degrees", "45"), LOCTEXT("Rotate45Degrees", "90"), LOCTEXT("RotateCustom", "---") };
 	const static int32 RotationValues[] = {45, 90, -1};
 	const static RotationAction RotationActions[] = {Rotate,Rotate,RotateCustom};
 	
@@ -303,7 +303,7 @@ TSharedRef<SWidget> SSurfaceProperties::ConstructTextureRotate()
 					[
 						SNew(STextBlock)
 						.Text(ButtonFields[Idx])
-						.ToolTipText(LOCTEXT("RotateToolTip_Custom", "Sets a custom rotate amount").ToString())
+						.ToolTipText(LOCTEXT("RotateToolTip_Custom", "Sets a custom rotate amount"))
 					]
 				] 
 				.MenuContent()
@@ -331,7 +331,7 @@ TSharedRef<SWidget> SSurfaceProperties::ConstructTextureRotate()
 				.VAlign(VAlign_Center)
 				.HAlign(HAlign_Center)
 				.Text(ButtonFields[Idx])
-				.ToolTipText(LOCTEXT("RotateToolTip", "Rotate texture").ToString())
+				.ToolTipText(LOCTEXT("RotateToolTip", "Rotate texture"))
 				.OnClicked(this,&SSurfaceProperties::OnRotateTexture,RotationValues[Idx],RotationActions[Idx])
 			];
 		}
@@ -351,7 +351,7 @@ TSharedRef<SWidget> SSurfaceProperties::ConstructTextureFlip()
 	.Padding(0,0,5,5)
 	[
 		SNew (STextBlock)
-		.Text(LOCTEXT("FlipTitle", "Flip:").ToString())
+		.Text(LOCTEXT("FlipTitle", "Flip:"))
 	];
 
 	Parent->AddSlot()
@@ -361,7 +361,7 @@ TSharedRef<SWidget> SSurfaceProperties::ConstructTextureFlip()
 	];
 
 	// Flip button fields:
-	const static FString ButtonFields[] = {LOCTEXT("RotateFlipU", "Flip U").ToString(), LOCTEXT("RotateFlipV", "Flip V").ToString()};
+	const static FText ButtonFields[] = {LOCTEXT("RotateFlipU", "Flip U"), LOCTEXT("RotateFlipV", "Flip V")};
 	const static TextureCoordChannel TextureCoordinateChannels[] = {UChannel,VChannel};
 
 	// create flip controls
@@ -377,7 +377,7 @@ TSharedRef<SWidget> SSurfaceProperties::ConstructTextureFlip()
 			.HAlign(HAlign_Center)
 			.ContentPadding(FMargin(0,5))
 			.Text(ButtonFields[Idx])
-			.ToolTipText(LOCTEXT("FlipToolTip", "Flip texture").ToString())
+			.ToolTipText(LOCTEXT("FlipToolTip", "Flip texture"))
 			.OnClicked(this,&SSurfaceProperties::OnFlipTexture,TextureCoordinateChannels[Idx])
 		];
 	}
@@ -401,7 +401,7 @@ TSharedRef<SWidget> SSurfaceProperties::ConstructTextureScale()
 			[
 				SNew(SHyperlink)
 					.Text(this, &SSurfaceProperties::GetScalingLabel)
-					.ToolTipText(LOCTEXT("ScalingRelativeToggle", "Toggle between Absolute and Relative scaling").ToString())
+					.ToolTipText(LOCTEXT("ScalingRelativeToggle", "Toggle between Absolute and Relative scaling"))
 					.OnNavigate( this, &SSurfaceProperties::OnScaleLabelClicked )
 					.TextStyle(FEditorStyle::Get(), "DetailsView.HyperlinkStyle")
 			]
@@ -464,7 +464,7 @@ TSharedRef<SWidget> SSurfaceProperties::ConstructTextureScale()
 		.IsChecked( this, &SSurfaceProperties::IsPreserveScaleRatioChecked )
 		.OnCheckStateChanged( this, &SSurfaceProperties::OnPreserveScaleRatioToggled )
 		.Style( FEditorStyle::Get(), "TransparentCheckBox" )
-		.ToolTipText( LOCTEXT("PreserveScaleSurfaceToolTip", "When locked changes to ether scaling value will be applied to the other.").ToString() )
+		.ToolTipText( LOCTEXT("PreserveScaleSurfaceToolTip", "When locked changes to ether scaling value will be applied to the other.") )
 		[
 			SNew( SImage )
 			.Image( this, &SSurfaceProperties::GetPreserveScaleRatioImage )
@@ -480,8 +480,8 @@ TSharedRef<SWidget> SSurfaceProperties::ConstructTextureScale()
 	.HAlign(HAlign_Right)
 	[
 		SNew(SButton)
-		.Text(LOCTEXT("ApplyScaling", "Apply").ToString())
-		.ToolTipText(LOCTEXT("ApplyScalingToolTip", "Apply scaling to selected surfaces").ToString())
+		.Text(LOCTEXT("ApplyScaling", "Apply"))
+		.ToolTipText(LOCTEXT("ApplyScalingToolTip", "Apply scaling to selected surfaces"))
 		.OnClicked(this, &SSurfaceProperties::OnApplyScaling)
 	];
 
@@ -491,8 +491,6 @@ END_SLATE_FUNCTION_BUILD_OPTIMIZATION
 
 TSharedRef<SWidget> SSurfaceProperties::ConstructLighting()
 {
-	FString MaterialNameForTitle;
-
 	LevelLightmassSettingsObjects.Empty();
 	SelectedLightmassSettingsObjects.Empty();
 	
@@ -528,30 +526,6 @@ TSharedRef<SWidget> SSurfaceProperties::ConstructLighting()
 						ObjArray.Add(LightmassSettingsObject);
 						SelectedLightmassSettingsObjects.Add(LightmassSettingsObject);
 					}
-
-					//Find Material Name for title
-					FString NewName;
-					if (Surf.Material)
-					{
-						NewName = Surf.Material->GetName();
-					}
-					else
-					{
-						NewName = NSLOCTEXT("UnrealEd", "SurfaceProperties_NoMaterial", "No Material").ToString();
-					}
-					//if we've never assigned a material before
-					if (MaterialNameForTitle.Len() == 0)
-					{
-						MaterialNameForTitle = NewName;
-					}
-					else
-					{
-						//already been assigned once, check if it's the same
-						if (NewName != MaterialNameForTitle)
-						{
-							MaterialNameForTitle = NSLOCTEXT("UnrealEd", "SurfaceProperties_MultipleMaterials", "Multiple Materials").ToString();
-						}
-					}
 				}
 			}
 			LevelLightmassSettingsObjects.Add(ObjArray);
@@ -566,7 +540,7 @@ TSharedRef<SWidget> SSurfaceProperties::ConstructLighting()
 	.Padding(0,0,0,5)
 	[
 		SNew(STextBlock)
-		.Text(LOCTEXT("LightingTitle", "Lighting:").ToString())
+		.Text(LOCTEXT("LightingTitle", "Lighting:"))
 	];
 
 	Parent->AddSlot()
@@ -578,7 +552,7 @@ TSharedRef<SWidget> SSurfaceProperties::ConstructLighting()
 		.Padding(0,0,10,5)
 		[
 			SNew(STextBlock)
-			.Text(LOCTEXT("LightingLightMapResolution", "Lightmap Resolution:").ToString())
+			.Text(LOCTEXT("LightingLightMapResolution", "Lightmap Resolution:"))
 		]
 
 		+SHorizontalBox::Slot()
@@ -588,7 +562,7 @@ TSharedRef<SWidget> SSurfaceProperties::ConstructLighting()
 		[
 			SNew(SNumericEntryBox<float>)
 			.OnValueCommitted(this, &SSurfaceProperties::OnLightmapResolutionCommitted)
-			.UndeterminedString( LOCTEXT("MultipleValues", "Multiple Values").ToString() )
+			.UndeterminedString( LOCTEXT("MultipleValues", "Multiple Values") )
 			.Value(this, &SSurfaceProperties::GetLightmapResolutionValue)
 		]
 
@@ -874,9 +848,9 @@ void SSurfaceProperties::OnScaleLabelClicked( )
 	GConfig->SetBool(TEXT("SurfaceSelection"), TEXT("UseRelativeScaling"), bUseRelativeScaling, GEditorUserSettingsIni);
 }
 
-FString SSurfaceProperties::GetScalingLabel() const
+FText SSurfaceProperties::GetScalingLabel() const
 {
-	return (bUseRelativeScaling) ? LOCTEXT("ScaleRelativeTitle", "Scale Relative:").ToString() : LOCTEXT("ScaleTitle", "Scale:").ToString();
+	return (bUseRelativeScaling) ? LOCTEXT("ScaleRelativeTitle", "Scale Relative:") : LOCTEXT("ScaleTitle", "Scale:");
 }
 
 TSharedRef< ITableRow > SSurfaceProperties::OnGenerateScaleTableRow( TSharedPtr<FString> Item, const TSharedRef< STableViewBase >& OwnerTable )
diff --git a/Engine/Source/Editor/LevelEditor/Private/SSurfaceProperties.h b/Engine/Source/Editor/LevelEditor/Private/SSurfaceProperties.h
--- a/Engine/Source/Editor/LevelEditor/Private/SSurfaceProperties.h
+++ b/Engine/Source/Editor/LevelEditor/Private/SSurfaceProperties.h
@@ -185,7 +185,7 @@ private:
 	/** 
 	 * Returns the text for the scaling label, determined by the value of bUseRelativeScaling
 	 */
-	FString GetScalingLabel() const;
+	FText GetScalingLabel() const;
 
 	/** 
 	 * Called with a value has been committed by one of the custom pan combobuttons
diff --git a/Engine/Source/Editor/Levels/Private/LevelViewModel.cpp b/Engine/Source/Editor/Levels/Private/LevelViewModel.cpp
--- a/Engine/Source/Editor/Levels/Private/LevelViewModel.cpp
+++ b/Engine/Source/Editor/Levels/Private/LevelViewModel.cpp
@@ -508,11 +508,11 @@ void FLevelViewModel::ChangeColor(const TSharedRef<SWidget>& InPickerParentWidge
 }
 
 //@todo: add actor drag+drop support
-bool FLevelViewModel::CanAssignActors( const TArray< TWeakObjectPtr<AActor> > Actors, FString& OutMessage ) const
+bool FLevelViewModel::CanAssignActors( const TArray< TWeakObjectPtr<AActor> > Actors, FText& OutMessage ) const
 {
 	return false;
 }
-bool FLevelViewModel::CanAssignActor( const TWeakObjectPtr<AActor> Actor, FString& OutMessage ) const
+bool FLevelViewModel::CanAssignActor(const TWeakObjectPtr<AActor> Actor, FText& OutMessage) const
 {
 	return true;
 }
diff --git a/Engine/Source/Editor/Levels/Private/LevelViewModel.h b/Engine/Source/Editor/Levels/Private/LevelViewModel.h
--- a/Engine/Source/Editor/Levels/Private/LevelViewModel.h
+++ b/Engine/Source/Editor/Levels/Private/LevelViewModel.h
@@ -167,7 +167,7 @@ public:
 	 *	@return						if true then at least one actor can be assigned; 
 	 *								If false, either Invalid actors were discovered or all actors are already assigned
 	 */
-	bool CanAssignActors( const TArray< TWeakObjectPtr<AActor> > Actors, FString& OutMessage ) const;
+	bool CanAssignActors(const TArray< TWeakObjectPtr<AActor> > Actors, FText& OutMessage) const;
 
 	/**
 	 *	Returns whether the specified actor can be assigned to the Level
@@ -177,7 +177,7 @@ public:
 	 *	@return						if true then at least one actor can be assigned; 
 	 *								If false, either the Actor is Invalid or already assigned
 	 */
-	bool CanAssignActor( const TWeakObjectPtr<AActor> Actor, FString& OutMessage ) const;
+	bool CanAssignActor(const TWeakObjectPtr<AActor> Actor, FText& OutMessage) const;
 	
 	/**
 	 *	Appends all of the actors associated with this Level to the specified list
diff --git a/Engine/Source/Editor/Levels/Private/SLevelsView.h b/Engine/Source/Editor/Levels/Private/SLevelsView.h
--- a/Engine/Source/Editor/Levels/Private/SLevelsView.h
+++ b/Engine/Source/Editor/Levels/Private/SLevelsView.h
@@ -68,7 +68,7 @@ public:
 				.HeaderContent()
 				[
 					SNew(STextBlock)
-						.ToolTipText( NSLOCTEXT("LevelsView", "Lock", "Lock").ToString() )
+						.ToolTipText( NSLOCTEXT("LevelsView", "Lock", "Lock") )
 				]
 
 			/** Level color column */
@@ -77,7 +77,7 @@ public:
 				.HeaderContent()
 				[
 					SNew(STextBlock)
-						.ToolTipText( NSLOCTEXT("LevelsView", "Color", "Color").ToString() )
+						.ToolTipText( NSLOCTEXT("LevelsView", "Color", "Color") )
 				]
 
 			/** Level visibility column */
@@ -86,7 +86,7 @@ public:
 				.HeaderContent()
 				[
 					SNew(STextBlock)
-						.ToolTipText( NSLOCTEXT("LevelsView", "Visibility", "Visibility").ToString() )
+						.ToolTipText( NSLOCTEXT("LevelsView", "Visibility", "Visibility") )
 				]
 
 			/** Level kismet column */
@@ -95,12 +95,12 @@ public:
 				.HeaderContent()
 				[
 					SNew(STextBlock)
-						.ToolTipText( NSLOCTEXT("LevelsView", "Kismet", "Open the level blueprint for this Level").ToString() )
+						.ToolTipText( NSLOCTEXT("LevelsView", "Kismet", "Open the level blueprint for this Level") )
 				]
 
 			/** LevelName label column */
 			+ SHeaderRow::Column( LevelsView::ColumnID_LevelLabel )
-				.DefaultLabel( LOCTEXT("Column_LevelNameLabel", "Level").ToString() )
+				.DefaultLabel( LOCTEXT("Column_LevelNameLabel", "Level") )
 				.FillWidth( 0.45f )
 
 			/** Level SCC status column */
@@ -109,7 +109,7 @@ public:
 				.HeaderContent()
 				[
 					SNew(STextBlock)
-						.ToolTipText( NSLOCTEXT("LevelsView", "SCCStatus", "Status in Source Control").ToString() )
+						.ToolTipText( NSLOCTEXT("LevelsView", "SCCStatus", "Status in Source Control") )
 				]
 
 			/** Level save column */
@@ -118,7 +118,7 @@ public:
 				.HeaderContent()
 				[
 					SNew(STextBlock)
-						.ToolTipText( NSLOCTEXT("LevelsView", "Save", "Save this Level").ToString() )
+						.ToolTipText( NSLOCTEXT("LevelsView", "Save", "Save this Level") )
 				];
 
 		ChildSlot
@@ -360,8 +360,8 @@ private:
 					.VAlign(VAlign_Center)
 					[
 						SNew(STextBlock)
-						.Text( LOCTEXT("Column_ActorCountLabel", "Actors").ToString() )
-						.ToolTipText( LOCTEXT("Column_ActorCountLabel", "Actors").ToString() )
+						.Text( LOCTEXT("Column_ActorCountLabel", "Actors") )
+						.ToolTipText( LOCTEXT("Column_ActorCountLabel", "Actors") )
 					]
 				], InsertColumnIndex);
 			++InsertColumnIndex;
@@ -389,8 +389,8 @@ private:
 					.VAlign(VAlign_Center)
 					[
 						SNew(STextBlock)
-						.Text( LOCTEXT("Column_LightmassSizeLabel", "Lightmass Size (MB)").ToString() )
-						.ToolTipText( LOCTEXT("Column_LightmassSizeLabel", "Lightmass Size (MB)").ToString() )
+						.Text( LOCTEXT("Column_LightmassSizeLabel", "Lightmass Size (MB)") )
+						.ToolTipText( LOCTEXT("Column_LightmassSizeLabel", "Lightmass Size (MB)") )
 					]
 				], InsertColumnIndex);
 			++InsertColumnIndex;
@@ -418,8 +418,8 @@ private:
 					.VAlign(VAlign_Center)
 					[
 						SNew(STextBlock)
-						.Text( LOCTEXT("Column_FileSizeLabel", "File Size (MB)").ToString() )
-						.ToolTipText( LOCTEXT("Column_FileSizeLabel", "File Size (MB)").ToString() )
+						.Text( LOCTEXT("Column_FileSizeLabel", "File Size (MB)") )
+						.ToolTipText( LOCTEXT("Column_FileSizeLabel", "File Size (MB)") )
 					]
 				], InsertColumnIndex);
 			++InsertColumnIndex;
@@ -447,7 +447,7 @@ private:
 					.VAlign(VAlign_Center)
 					[
 						SNew(STextBlock)
-						.Text( LOCTEXT("Column_EditorOffsetLabel", "Editor Offset").ToString() )
+						.Text( LOCTEXT("Column_EditorOffsetLabel", "Editor Offset") )
 					]
 				], InsertColumnIndex);
 			++InsertColumnIndex;
diff --git a/Engine/Source/Editor/Levels/Private/SLevelsViewRow.h b/Engine/Source/Editor/Levels/Private/SLevelsViewRow.h
--- a/Engine/Source/Editor/Levels/Private/SLevelsViewRow.h
+++ b/Engine/Source/Editor/Levels/Private/SLevelsViewRow.h
@@ -404,7 +404,7 @@ protected:
 		}
 
 		bool bCanAssign = false;
-		FString Message;
+		FText Message;
 		if( DragActorOp->Actors.Num() > 1 )
 		{
 			bCanAssign = ViewModel->CanAssignActors( DragActorOp->Actors, OUT Message );
@@ -416,11 +416,11 @@ protected:
 
 		if ( bCanAssign )
 		{
-			DragActorOp->SetToolTip( FActorDragDropGraphEdOp::ToolTip_CompatibleGeneric, FText::FromString(Message) );
+			DragActorOp->SetToolTip( FActorDragDropGraphEdOp::ToolTip_CompatibleGeneric, Message );
 		}
 		else
 		{
-			DragActorOp->SetToolTip( FActorDragDropGraphEdOp::ToolTip_IncompatibleGeneric, FText::FromString(Message) );
+			DragActorOp->SetToolTip( FActorDragDropGraphEdOp::ToolTip_IncompatibleGeneric, Message );
 		}
 
 		return FReply::Handled();
@@ -513,7 +513,7 @@ private:
 		{
 			TSharedPtr<FActorDragDropGraphEdOp> DragDropActorOp = StaticCastSharedPtr<FActorDragDropGraphEdOp>(DragDropOp);
 
-			FString Message;
+			FText Message;
 			bCanAcceptDrop = ViewModel->CanAssignActors( DragDropActorOp->Actors, OUT Message );
 		}
 
diff --git a/Engine/Source/Editor/MaterialEditor/Private/SMaterialPalette.cpp b/Engine/Source/Editor/MaterialEditor/Private/SMaterialPalette.cpp
--- a/Engine/Source/Editor/MaterialEditor/Private/SMaterialPalette.cpp
+++ b/Engine/Source/Editor/MaterialEditor/Private/SMaterialPalette.cpp
@@ -115,7 +115,7 @@ void SMaterialPalette::Construct(const FArguments& InArgs, TWeakPtr<FMaterialEdi
 				.AutoWidth()
 				[
 					SNew(STextBlock)
-					.Text(LOCTEXT("Category", "Category: ").ToString())
+					.Text(LOCTEXT("Category", "Category: "))
 				]
 
 				// Combo button to select a class
diff --git a/Engine/Source/Editor/PListEditor/Private/SPlistEditor.cpp b/Engine/Source/Editor/PListEditor/Private/SPlistEditor.cpp
--- a/Engine/Source/Editor/PListEditor/Private/SPlistEditor.cpp
+++ b/Engine/Source/Editor/PListEditor/Private/SPlistEditor.cpp
@@ -61,7 +61,7 @@ void SPListEditorPanel::Construct( const FArguments& InArgs )
 							[
 								SNew(SExpandableArea)
 								.InitiallyCollapsed(true)
-								.AreaTitle(LOCTEXT("PListMenuTitle", "File").ToString())
+								.AreaTitle(LOCTEXT("PListMenuTitle", "File"))
 								.BodyContent()
 								[
 									SNew(SHorizontalBox)
@@ -73,8 +73,8 @@ void SPListEditorPanel::Construct( const FArguments& InArgs )
 									.HAlign(HAlign_Left)
 									[
 										SNew(SButton)
-										.Text(LOCTEXT("PListNew", "New").ToString())
-										.ToolTipText(LOCTEXT("PListNewToolTip", "Create a new plist file").ToString())
+										.Text(LOCTEXT("PListNew", "New"))
+										.ToolTipText(LOCTEXT("PListNewToolTip", "Create a new plist file"))
 										.OnClicked(this, &SPListEditorPanel::OnNewClicked)
 									]
 
@@ -85,8 +85,8 @@ void SPListEditorPanel::Construct( const FArguments& InArgs )
 									.HAlign(HAlign_Left)
 									[
 										SNew(SButton)
-										.Text(LOCTEXT("PListOpen", "Open...").ToString())
-										.ToolTipText(LOCTEXT("PListOpenToolTip", "Open an existing plist file").ToString())
+										.Text(LOCTEXT("PListOpen", "Open..."))
+										.ToolTipText(LOCTEXT("PListOpenToolTip", "Open an existing plist file"))
 										.OnClicked(this, &SPListEditorPanel::OnOpenClicked)
 									]
 
@@ -97,8 +97,8 @@ void SPListEditorPanel::Construct( const FArguments& InArgs )
 									.HAlign(HAlign_Left)
 									[
 										SNew(SButton)
-										.Text(LOCTEXT("PListSave", "Save").ToString())
-										.ToolTipText(LOCTEXT("PListSaveToolTip", "Save current working plist").ToString())
+										.Text(LOCTEXT("PListSave", "Save"))
+										.ToolTipText(LOCTEXT("PListSaveToolTip", "Save current working plist"))
 										.OnClicked(this, &SPListEditorPanel::OnSaveClicked)
 									]
 
@@ -109,8 +109,8 @@ void SPListEditorPanel::Construct( const FArguments& InArgs )
 									.HAlign(HAlign_Left)
 									[
 										SNew(SButton)
-										.Text(LOCTEXT("PListSaveAs", "Save As...").ToString())
-										.ToolTipText(LOCTEXT("PListSaveAsToolTip", "Save current working plist with a specified filename").ToString())
+										.Text(LOCTEXT("PListSaveAs", "Save As..."))
+										.ToolTipText(LOCTEXT("PListSaveAsToolTip", "Save current working plist with a specified filename"))
 										.OnClicked(this, &SPListEditorPanel::OnSaveAsClicked)
 									]
 								]
@@ -141,8 +141,8 @@ void SPListEditorPanel::Construct( const FArguments& InArgs )
 							//.HAlign(HAlign_Left)
 							//[
 							//	SAssignNew(RemoveSelectedButton, SButton)
-							//	.Text(LOCTEXT("PListRemoveSelected", "Remove Selected Rows").ToString())
-							//	.ToolTipText(LOCTEXT("PListRemoveSelectedToolTip", "Removes the selected rows from the plist below").ToString())
+							//	.Text(LOCTEXT("PListRemoveSelected", "Remove Selected Rows"))
+							//	.ToolTipText(LOCTEXT("PListRemoveSelectedToolTip", "Removes the selected rows from the plist below"))
 							//	.OnClicked(this, &SPListEditorPanel::OnDeleteSelectedClicked)
 							//	.IsEnabled(false)
 							//]
@@ -179,7 +179,7 @@ void SPListEditorPanel::Construct( const FArguments& InArgs )
 						.HeaderContentPadding(FMargin(6))
 						[
 							SNew(STextBlock)
-							.Text(LOCTEXT("PListKeySectionTitle", "Key").ToString())
+							.Text(LOCTEXT("PListKeySectionTitle", "Key"))
 						]
 
 						+ SHeaderRow::Column(FName(TEXT("PListValueTypeColumn")))
@@ -187,7 +187,7 @@ void SPListEditorPanel::Construct( const FArguments& InArgs )
 						.HeaderContentPadding(FMargin(6))
 						[
 							SNew(STextBlock)
-							.Text(LOCTEXT("PListValueTypeSectionTitle", "Value Type").ToString())
+							.Text(LOCTEXT("PListValueTypeSectionTitle", "Value Type"))
 						]
 
 						+ SHeaderRow::Column(FName(TEXT("PListValueColumn")))
@@ -195,7 +195,7 @@ void SPListEditorPanel::Construct( const FArguments& InArgs )
 						.HeaderContentPadding(FMargin(6))
 						[
 							SNew(STextBlock)
-							.Text(LOCTEXT("PListValueSectionTitle", "Value").ToString())
+							.Text(LOCTEXT("PListValueSectionTitle", "Value"))
 						]
 					)
 				]
@@ -906,13 +906,12 @@ FReply SPListEditorPanel::OnOpenClicked()
 			ParentWindowWindowHandle = MainFrameParentWindow->GetNativeWindow()->GetOSWindowHandle();
 		}
 
-		FString Title = LOCTEXT("PListOpenDialogTitle", "Open").ToString();
 		//FString DefaultPath = FPaths::GameDir() + TEXT("Build/IOS/");
 		//FString DefaultFile = FString(GGameName) + TEXT("-Info.plist");
 		FString FileTypes = TEXT("Property List (*.plist)|*.plist|All Files (*.*)|*.*");
 		DesktopPlatform->OpenFileDialog(
 			ParentWindowWindowHandle,
-			Title,
+			LOCTEXT("PListOpenDialogTitle", "Open").ToString(),
 			InOutLastPath,
 			TEXT(""),
 			FileTypes,
@@ -1030,11 +1029,10 @@ bool SPListEditorPanel::PromptSave()
 					ParentWindowWindowHandle = MainFrameParentWindow->GetNativeWindow()->GetOSWindowHandle();
 				}
 			
-				FString Title = LOCTEXT("PListSaveDialogTitle", "Save").ToString();
 				FString FileTypes = TEXT("Property List (*.plist)|*.plist|All Files (*.*)|*.*");
 				DesktopPlatform->SaveFileDialog(
 					ParentWindowWindowHandle,
-					Title,
+					LOCTEXT("PListSaveDialogTitle", "Save").ToString(),
 					InOutLastPath,
 					TEXT(""),
 					FileTypes,
@@ -1177,11 +1175,10 @@ FReply SPListEditorPanel::OnSaveClicked()
 				ParentWindowWindowHandle = MainFrameParentWindow->GetNativeWindow()->GetOSWindowHandle();
 			}
 			
-			FString Title = LOCTEXT("PListSaveDialogTitle", "Save").ToString();
 			FString FileTypes = TEXT("Property List (*.plist)|*.plist|All Files (*.*)|*.*");
 			DesktopPlatform->SaveFileDialog(
 				ParentWindowWindowHandle,
-				Title,
+				LOCTEXT("PListSaveDialogTitle", "Save").ToString(),
 				InOutLastPath,
 				TEXT(""),
 				FileTypes,
@@ -1292,11 +1289,10 @@ FReply SPListEditorPanel::OnSaveAsClicked()
 			ParentWindowWindowHandle = MainFrameParentWindow->GetNativeWindow()->GetOSWindowHandle();
 		}
 			
-		FString Title = LOCTEXT("PListSaveAsDialogTitle", "Save As").ToString();
 		FString FileTypes = TEXT("Property List (*.plist)|*.plist|All Files (*.*)|*.*");
 		DesktopPlatform->SaveFileDialog(
 			ParentWindowWindowHandle,
-			Title,
+			LOCTEXT("PListSaveAsDialogTitle", "Save As").ToString(),
 			InOutLastPath,
 			TEXT(""),
 			FileTypes,
diff --git a/Engine/Source/Editor/Persona/Private/AnimationMode/AnimationMode.cpp b/Engine/Source/Editor/Persona/Private/AnimationMode/AnimationMode.cpp
--- a/Engine/Source/Editor/Persona/Private/AnimationMode/AnimationMode.cpp
+++ b/Engine/Source/Editor/Persona/Private/AnimationMode/AnimationMode.cpp
@@ -237,7 +237,7 @@ void SAnimDifferentAssetBeingPreviewedWarning::Construct(const FArguments& InArg
 		[
 			SNew(STextBlock)
 			.Font(BoldFont)
-			.Text(LOCTEXT("AssetPreviewMismatch", "Currently previewing\nanother asset").ToString())
+			.Text(LOCTEXT("AssetPreviewMismatch", "Currently previewing\nanother asset"))
 		]
 	];
 }
diff --git a/Engine/Source/Editor/Persona/Private/MeshMode/SAdditionalMeshesEditor.cpp b/Engine/Source/Editor/Persona/Private/MeshMode/SAdditionalMeshesEditor.cpp
--- a/Engine/Source/Editor/Persona/Private/MeshMode/SAdditionalMeshesEditor.cpp
+++ b/Engine/Source/Editor/Persona/Private/MeshMode/SAdditionalMeshesEditor.cpp
@@ -176,7 +176,7 @@ void SAdditionalMeshesEditor::Construct(const FArguments& InArgs, TSharedPtr<FPe
 				.OnGetMenuContent(this, &SAdditionalMeshesEditor::MakeAssetPickerMenu)
 				.ButtonContent()
 				[
-					SNew(STextBlock) .Text(LOCTEXT("AdditionalMeshesAddButtonLabel", "Add Mesh").ToString())
+					SNew(STextBlock) .Text(LOCTEXT("AdditionalMeshesAddButtonLabel", "Add Mesh"))
 				]
 			]
 			+SHorizontalBox::Slot()
@@ -184,7 +184,7 @@ void SAdditionalMeshesEditor::Construct(const FArguments& InArgs, TSharedPtr<FPe
 			[
 				SNew(SButton)
 				.OnClicked(this, &SAdditionalMeshesEditor::OnClearAllClicked)
-				.Text(LOCTEXT("AdditionalMeshesClearButtonLabel", "Clear All").ToString())
+				.Text(LOCTEXT("AdditionalMeshesClearButtonLabel", "Clear All"))
 			]
 		]
 		+SVerticalBox::Slot()
diff --git a/Engine/Source/Editor/Persona/Private/PersonaMeshDetails.cpp b/Engine/Source/Editor/Persona/Private/PersonaMeshDetails.cpp
--- a/Engine/Source/Editor/Persona/Private/PersonaMeshDetails.cpp
+++ b/Engine/Source/Editor/Persona/Private/PersonaMeshDetails.cpp
@@ -100,11 +100,11 @@ TSharedRef<SWidget> FPersonaMeshDetails::OnGenerateCustomNameWidgetsForMaterial(
 		SNew(SCheckBox)
 		.IsChecked(this, &FPersonaMeshDetails::IsSectionSelected, SlotIndex)
 		.OnCheckStateChanged(this, &FPersonaMeshDetails::OnSectionSelectedChanged, SlotIndex)
-		.ToolTipText(LOCTEXT("Highlight_ToolTip", "Highlights this section in the viewport").ToString())
+		.ToolTipText(LOCTEXT("Highlight_ToolTip", "Highlights this section in the viewport"))
 		[
 			SNew(STextBlock)
 			.Font(IDetailLayoutBuilder::GetDetailFont())
-			.Text(LOCTEXT("Highlight", "Highlight").ToString())
+			.Text(LOCTEXT("Highlight", "Highlight"))
 			.ColorAndOpacity( FLinearColor( 0.4f, 0.4f, 0.4f, 1.0f) )
 		];
 }
@@ -137,7 +137,7 @@ TSharedRef<SWidget> FPersonaMeshDetails::OnGenerateCustomMaterialWidgetsForMater
 			[
 				SNew(STextBlock)
 				.Font(FEditorStyle::GetFontStyle("StaticMeshEditor.NormalFont"))
-				.Text(LOCTEXT("Cast Shadows", "Cast Shadows").ToString())
+				.Text(LOCTEXT("Cast Shadows", "Cast Shadows"))
 			]
 		]
 
@@ -166,7 +166,7 @@ TSharedRef<SWidget> FPersonaMeshDetails::OnGenerateCustomMaterialWidgetsForMater
 			[
 				SNew(STextBlock)
 				.Font(FEditorStyle::GetFontStyle("StaticMeshEditor.NormalFont"))
-				.Text(LOCTEXT("PersonaDeleteMaterialLabel", "Delete").ToString())
+				.Text(LOCTEXT("PersonaDeleteMaterialLabel", "Delete"))
 			]
 	]
 
@@ -185,7 +185,7 @@ TSharedRef<SWidget> FPersonaMeshDetails::OnGenerateCustomMaterialWidgetsForMater
 			.VAlign(VAlign_Center)
 			[
 				SNew(STextBlock)
-				.Text(LOCTEXT("Clothing", "Clothing").ToString())
+				.Text(LOCTEXT("Clothing", "Clothing"))
 			]
 
 			+SHorizontalBox::Slot()
@@ -195,7 +195,7 @@ TSharedRef<SWidget> FPersonaMeshDetails::OnGenerateCustomMaterialWidgetsForMater
 				SAssignNew(ClothingComboBoxes[SlotIndex], STextComboBox)
 				.ContentPadding(FMargin(6.0f, 2.0f))
 				.OptionsSource(&ClothingComboStrings)
-				.ToolTipText(LOCTEXT("SectionsComboBoxToolTip", "Select the clothing asset and submesh to use as clothing for this section").ToString())
+				.ToolTipText(LOCTEXT("SectionsComboBoxToolTip", "Select the clothing asset and submesh to use as clothing for this section"))
 				.OnSelectionChanged(this, &FPersonaMeshDetails::HandleSectionsComboBoxSelectionChanged, SlotIndex)
 			]
 		]
@@ -212,8 +212,8 @@ TSharedRef<SWidget> FPersonaMeshDetails::OnGenerateCustomMaterialWidgetsForMater
 			[
 				SNew( STextBlock )
 				.Font(FEditorStyle::GetFontStyle("StaticMeshEditor.NormalFont"))
-				.Text( LOCTEXT("EnableClothLODCheckBox", "Enable Clothing LOD").ToString() )
-				.ToolTipText( LOCTEXT("UseLODTooltip", "If enabled, clothing data will be used for each LOD if there is a corresponding APEX clothing LOD available. If disabled, clothing data will only be used when LOD 0 is rendered.").ToString() )
+				.Text( LOCTEXT("EnableClothLODCheckBox", "Enable Clothing LOD") )
+				.ToolTipText( LOCTEXT("UseLODTooltip", "If enabled, clothing data will be used for each LOD if there is a corresponding APEX clothing LOD available. If disabled, clothing data will only be used when LOD 0 is rendered.") )
 			]
 		];
 
@@ -370,8 +370,8 @@ void FPersonaMeshDetails::CustomizeClothingProperties(IDetailLayoutBuilder& Deta
 		[
 			SNew(SButton)
 			.OnClicked(this, &FPersonaMeshDetails::OnOpenClothingFileClicked, &DetailLayout)
-			.Text(LOCTEXT("AddAPEXClothingFile", "Add APEX clothing file...").ToString())
-			.ToolTipText(LOCTEXT("AddAPEXClothingFileTip", "Select a new APEX clothing file and add it to the skeletal mesh.").ToString())
+			.Text(LOCTEXT("AddAPEXClothingFile", "Add APEX clothing file..."))
+			.ToolTipText(LOCTEXT("AddAPEXClothingFileTip", "Select a new APEX clothing file and add it to the skeletal mesh."))
 		]
 	];
 }
@@ -397,13 +397,13 @@ void FPersonaMeshDetails::OnGenerateElementForClothingAsset( TSharedRef<IPropert
 		.AutoWidth()
 		[
 			SNew( SButton )
-			.Text( LOCTEXT("ReimportButtonLabel", "Reimport").ToString() )
+			.Text( LOCTEXT("ReimportButtonLabel", "Reimport") )
 			.OnClicked( this, &FPersonaMeshDetails::OnReimportApexFileClicked, ElementIndex )
 			.IsFocusable( false )
 			.ContentPadding(0)
 			.ForegroundColor( FSlateColor::UseForeground() )
 			.ButtonColorAndOpacity(FLinearColor(1.0f,1.0f,1.0f,0.0f))
-			.ToolTipText(LOCTEXT("ReimportApexFileTip", "Reimport this APEX asset").ToString())
+			.ToolTipText(LOCTEXT("ReimportApexFileTip", "Reimport this APEX asset"))
 			[ 
 				SNew( SImage )
 				.Image( FEditorStyle::GetBrush("Persona.ReimportAsset") )
@@ -418,13 +418,13 @@ void FPersonaMeshDetails::OnGenerateElementForClothingAsset( TSharedRef<IPropert
 		.AutoWidth()
 		[
 			SNew( SButton )
-			.Text( LOCTEXT("ClearButtonLabel", "Remove").ToString() )
+			.Text( LOCTEXT("ClearButtonLabel", "Remove") )
 			.OnClicked( this, &FPersonaMeshDetails::OnRemoveApexFileClicked, ElementIndex, DetailLayout )
 			.IsFocusable( false )
 			.ContentPadding(0)
 			.ForegroundColor( FSlateColor::UseForeground() )
 			.ButtonColorAndOpacity(FLinearColor(1.0f,1.0f,1.0f,0.0f))
-			.ToolTipText(LOCTEXT("RemoveApexFileTip", "Remove this APEX asset").ToString())
+			.ToolTipText(LOCTEXT("RemoveApexFileTip", "Remove this APEX asset"))
 			[ 
 				SNew( SImage )
 				.Image( FEditorStyle::GetBrush("PropertyWindow.Button_Clear") )
@@ -446,7 +446,7 @@ void FPersonaMeshDetails::OnGenerateElementForClothingAsset( TSharedRef<IPropert
 	.NameContent()
 	[
 		SNew(STextBlock)
-		.Text(LOCTEXT("APEXDetails", "APEX Details").ToString())
+		.Text(LOCTEXT("APEXDetails", "APEX Details"))
 		.Font(DetailFontInfo)
 	]
 	.ValueContent()
@@ -474,7 +474,7 @@ TSharedRef<SUniformGridPanel> FPersonaMeshDetails::MakeApexDetailsWidget(int32 A
 			[
 				SNew(STextBlock)
 				.Font(DetailFontInfo)
-				.Text(LOCTEXT("Removing", "Removing...").ToString())
+				.Text(LOCTEXT("Removing", "Removing..."))
 			];
 
 		return Grid;
@@ -490,7 +490,7 @@ TSharedRef<SUniformGridPanel> FPersonaMeshDetails::MakeApexDetailsWidget(int32 A
 		[
 			SNew(STextBlock)
 			.Font(DetailFontInfo)
-			.Text(FString::Printf(*LOCTEXT("LODIndex", "LOD %d").ToString(), LODIndex))			
+			.Text(FText::Format(LOCTEXT("LODIndex", "LOD {0}"), FText::AsNumber(LODIndex)))			
 		];
 
 		RowNumber++;
@@ -504,7 +504,7 @@ TSharedRef<SUniformGridPanel> FPersonaMeshDetails::MakeApexDetailsWidget(int32 A
 			[
 				SNew(STextBlock)
 				.Font(DetailFontInfo)
-				.Text(LOCTEXT("SubmeshIndex", "Submesh").ToString())
+				.Text(LOCTEXT("SubmeshIndex", "Submesh"))
 			];
 
 			Grid->AddSlot(1, RowNumber) 
@@ -512,7 +512,7 @@ TSharedRef<SUniformGridPanel> FPersonaMeshDetails::MakeApexDetailsWidget(int32 A
 			[
 				SNew(STextBlock)
 				.Font(DetailFontInfo)
-				.Text(LOCTEXT("SimulVertexCount", "Simul Verts").ToString())
+				.Text(LOCTEXT("SimulVertexCount", "Simul Verts"))
 			];
 
 			Grid->AddSlot(2, RowNumber) 
@@ -520,7 +520,7 @@ TSharedRef<SUniformGridPanel> FPersonaMeshDetails::MakeApexDetailsWidget(int32 A
 			[
 				SNew(STextBlock)
 				.Font(DetailFontInfo)
-				.Text(LOCTEXT("RenderVertexCount", "Render Verts").ToString())
+				.Text(LOCTEXT("RenderVertexCount", "Render Verts"))
 			];
 
 			Grid->AddSlot(3, RowNumber)
@@ -528,7 +528,7 @@ TSharedRef<SUniformGridPanel> FPersonaMeshDetails::MakeApexDetailsWidget(int32 A
 			[
 				SNew(STextBlock)
 				.Font(DetailFontInfo)
-				.Text(LOCTEXT("FixedVertexCount", "Fixed Verts").ToString())
+				.Text(LOCTEXT("FixedVertexCount", "Fixed Verts"))
 			];
 
 			Grid->AddSlot(4, RowNumber)
@@ -536,7 +536,7 @@ TSharedRef<SUniformGridPanel> FPersonaMeshDetails::MakeApexDetailsWidget(int32 A
 			[
 				SNew(STextBlock)
 				.Font(DetailFontInfo)
-				.Text(LOCTEXT("TriangleCount", "Triangles").ToString())
+				.Text(LOCTEXT("TriangleCount", "Triangles"))
 			];
 
 			RowNumber++;
@@ -606,7 +606,7 @@ TSharedRef<SUniformGridPanel> FPersonaMeshDetails::MakeApexDetailsWidget(int32 A
 			[
 				SNew(STextBlock)
 				.Font(DetailFontInfo)
-				.Text(LOCTEXT("FailedGetSubmeshInfo", "Failed to get sub-mesh Info").ToString())
+				.Text(LOCTEXT("FailedGetSubmeshInfo", "Failed to get sub-mesh Info"))
 			];
 
 			RowNumber++;
diff --git a/Engine/Source/Editor/Persona/Private/SAnimCurvePanel.cpp b/Engine/Source/Editor/Persona/Private/SAnimCurvePanel.cpp
--- a/Engine/Source/Editor/Persona/Private/SAnimCurvePanel.cpp
+++ b/Engine/Source/Editor/Persona/Private/SAnimCurvePanel.cpp
@@ -480,10 +480,10 @@ FReply SCurveEdTrack::OnContextMenu()
 			SNew(SCheckBox)
 			.IsChecked( this, &SCurveEdTrack::IsCurveOfMode, ACF_DrivesMorphTarget )
 			.OnCheckStateChanged( this, &SCurveEdTrack::ToggleCurveMode, ACF_DrivesMorphTarget )
-			.ToolTipText(LOCTEXT("MorphCurveModeTooltip", "This curve drives a morph target").ToString())
+			.ToolTipText(LOCTEXT("MorphCurveModeTooltip", "This curve drives a morph target"))
 			[
 				SNew(STextBlock)
-				.Text(LOCTEXT("MorphCurveMode", "Morph Curve").ToString())
+				.Text(LOCTEXT("MorphCurveMode", "Morph Curve"))
 			],
 			FText()
 		);
@@ -492,10 +492,10 @@ FReply SCurveEdTrack::OnContextMenu()
 			SNew(SCheckBox)
 			.IsChecked( this, &SCurveEdTrack::IsCurveOfMode, ACF_DrivesMaterial )
 			.OnCheckStateChanged( this, &SCurveEdTrack::ToggleCurveMode, ACF_DrivesMaterial )
-			.ToolTipText(LOCTEXT("MaterialCurveModeTooltip", "This curve drives a material").ToString())
+			.ToolTipText(LOCTEXT("MaterialCurveModeTooltip", "This curve drives a material"))
 			[
 				SNew(STextBlock)
-				.Text(LOCTEXT("MaterialCurveMode", "Material Curve").ToString())
+				.Text(LOCTEXT("MaterialCurveMode", "Material Curve"))
 			],
 			FText()
 		);
@@ -686,10 +686,10 @@ FReply SAnimCurvePanel::OnContextMenu()
 			SNew(SCheckBox)
 			.IsChecked( this, &SAnimCurvePanel::AreAllCurvesOfMode, ACF_DrivesMorphTarget )
 			.OnCheckStateChanged( this, &SAnimCurvePanel::ToggleAllCurveModes, ACF_DrivesMorphTarget )
-			.ToolTipText(LOCTEXT("MorphCurveModeTooltip", "This curve drives a morph target").ToString())
+			.ToolTipText(LOCTEXT("MorphCurveModeTooltip", "This curve drives a morph target"))
 			[
 				SNew(STextBlock)
-				.Text(LOCTEXT("MorphCurveMode", "Morph Curve").ToString())
+				.Text(LOCTEXT("MorphCurveMode", "Morph Curve"))
 			],
 			FText()
 		);
@@ -698,11 +698,11 @@ FReply SAnimCurvePanel::OnContextMenu()
 			SNew(SCheckBox)
 			.IsChecked( this, &SAnimCurvePanel::AreAllCurvesOfMode, ACF_DrivesMaterial )
 			.OnCheckStateChanged( this, &SAnimCurvePanel::ToggleAllCurveModes, ACF_DrivesMaterial )
-			.ToolTipText(LOCTEXT("MaterialCurveModeTooltip", "This curve drives a material").ToString())
+			.ToolTipText(LOCTEXT("MaterialCurveModeTooltip", "This curve drives a material"))
 			.HAlign(HAlign_Left)
 			[
 				SNew(STextBlock)
-				.Text(LOCTEXT("MaterialCurveMode", "Material Curve").ToString())
+				.Text(LOCTEXT("MaterialCurveMode", "Material Curve"))
 			],
 			FText()
 		);
@@ -873,7 +873,7 @@ TSharedRef<SWidget> SAnimCurvePanel::GenerateCurveList()
 					.IsEnabled( true )
 					[
 						SNew( STextBlock )
-						.Text(Curve.CurveName.ToString())
+						.Text(FText::FromName(Curve.CurveName))
 					]
 				];
 		}
diff --git a/Engine/Source/Editor/Persona/Private/SAnimEditorBase.cpp b/Engine/Source/Editor/Persona/Private/SAnimEditorBase.cpp
--- a/Engine/Source/Editor/Persona/Private/SAnimEditorBase.cpp
+++ b/Engine/Source/Editor/Persona/Private/SAnimEditorBase.cpp
@@ -79,7 +79,7 @@ void SAnimEditorBase::Construct(const FArguments& InArgs)
 					.Padding(FMargin(4,4,0,0))
 					[
 						SNew( STextBlock )
-						.Text(LOCTEXT("Animation", "Animation : ").ToString())
+						.Text(LOCTEXT("Animation", "Animation : "))
 					]
 
 					+SHorizontalBox::Slot()
@@ -100,7 +100,7 @@ void SAnimEditorBase::Construct(const FArguments& InArgs)
 					.Padding(FMargin(4,4,0,0))
 					[
 						SNew( STextBlock )
-						.Text(LOCTEXT("Percentage", "Percentage: ").ToString())
+						.Text(LOCTEXT("Percentage", "Percentage: "))
 					]
 					+SHorizontalBox::Slot()
 					.FillWidth(1)
@@ -119,7 +119,7 @@ void SAnimEditorBase::Construct(const FArguments& InArgs)
 					.Padding(FMargin(4,4,0,0))
 					[
 						SNew( STextBlock )
-						.Text(LOCTEXT("CurrentTime", "CurrentTime: ").ToString())
+						.Text(LOCTEXT("CurrentTime", "CurrentTime: "))
 					]
 					+SHorizontalBox::Slot()
 					.FillWidth(1)
@@ -139,7 +139,7 @@ void SAnimEditorBase::Construct(const FArguments& InArgs)
 					.Padding(FMargin(4,4,0,0))
 					[
 						SNew( STextBlock )
-						.Text(LOCTEXT("CurrentFrame", "Current Frame: ").ToString())
+						.Text(LOCTEXT("CurrentFrame", "Current Frame: "))
 					]
 					+SHorizontalBox::Slot()
 					.FillWidth(1)
diff --git a/Engine/Source/Editor/Persona/Private/SAnimMontageSectionsPanel.cpp b/Engine/Source/Editor/Persona/Private/SAnimMontageSectionsPanel.cpp
--- a/Engine/Source/Editor/Persona/Private/SAnimMontageSectionsPanel.cpp
+++ b/Engine/Source/Editor/Persona/Private/SAnimMontageSectionsPanel.cpp
@@ -29,7 +29,7 @@ void SAnimMontageSectionsPanel::Construct(const FArguments& InArgs)
 		.FillHeight(1)
 		[
 			SNew( SExpandableArea )
-			.AreaTitle( LOCTEXT("Sections", "Sections").ToString() )
+			.AreaTitle( LOCTEXT("Sections", "Sections") )
 			.BodyContent()
 			[
 				SNew( SBorder )
@@ -89,8 +89,8 @@ void SAnimMontageSectionsPanel::Update()
 				[				
 					SNew(SButton)
 					.Visibility( EVisibility::Visible )
-					.Text( LOCTEXT("CreateDefault", "Create Default").ToString() )
-					.ToolTipText( LOCTEXT("CreateDefaultToolTip", "Reconstructs section ordering based on start time").ToString() )
+					.Text( LOCTEXT("CreateDefault", "Create Default") )
+					.ToolTipText( LOCTEXT("CreateDefaultToolTip", "Reconstructs section ordering based on start time") )
 					.OnClicked(this, &SAnimMontageSectionsPanel::MakeDefaultSequence)
 					.HAlign(HAlign_Center)
 					.VAlign(VAlign_Center)
@@ -102,8 +102,8 @@ void SAnimMontageSectionsPanel::Update()
 				[				
 					SNew(SButton)
 					.Visibility( EVisibility::Visible )
-					.Text( LOCTEXT("Clear", "Clear").ToString() )
-					.ToolTipText( LOCTEXT("ClearToolTip", "Resets section orderings").ToString() )
+					.Text( LOCTEXT("Clear", "Clear") )
+					.ToolTipText( LOCTEXT("ClearToolTip", "Resets section orderings") )
 					.OnClicked(this, &SAnimMontageSectionsPanel::ClearSequence)
 					.HAlign(HAlign_Center)
 					.VAlign(VAlign_Center)
@@ -157,8 +157,8 @@ void SAnimMontageSectionsPanel::Update()
 				[				
 					SNew(SButton)
 					.Visibility( EVisibility::Visible )
-					.Text( LOCTEXT("PreviewAll", "Preview All Sections").ToString() )
-					.ToolTipText( LOCTEXT("PreviewAllToolTip", "Preview all sections in order they are").ToString() )
+					.Text( LOCTEXT("PreviewAll", "Preview All Sections") )
+					.ToolTipText( LOCTEXT("PreviewAllToolTip", "Preview all sections in order they are") )
 					.OnClicked(this, &SAnimMontageSectionsPanel::PreviewAllSectionsClicked)
 					.HAlign(HAlign_Center)
 					.VAlign(VAlign_Center)
@@ -202,8 +202,8 @@ void SAnimMontageSectionsPanel::Update()
 					[				
 						SNew(SButton)
 						.Visibility( EVisibility::Visible )
-						.Text( LOCTEXT("Preview", "Preview").ToString() )
-						.ToolTipText( LOCTEXT("PreviewToolTip", "Preview this track").ToString() )
+						.Text( LOCTEXT("Preview", "Preview") )
+						.ToolTipText( LOCTEXT("PreviewToolTip", "Preview this track") )
 						.OnClicked(this, &SAnimMontageSectionsPanel::PreviewSectionClicked, SectionIdx)
 						.HAlign(HAlign_Center)
 						.VAlign(VAlign_Center)
diff --git a/Engine/Source/Editor/Persona/Private/SSkeletonTree.cpp b/Engine/Source/Editor/Persona/Private/SSkeletonTree.cpp
--- a/Engine/Source/Editor/Persona/Private/SSkeletonTree.cpp
+++ b/Engine/Source/Editor/Persona/Private/SSkeletonTree.cpp
@@ -301,7 +301,7 @@ void FDisplayedMeshBoneInfo::GenerateWidgetForNameColumn( TSharedPtr< SHorizonta
 		[
 			SNew( STextBlock )
 			.ColorAndOpacity(TextColor)
-			.Text( BoneName.ToString() )
+			.Text( FText::FromName(BoneName) )
 			.HighlightText( FilterText )
 			.Font(TextFont)
 			.ToolTipText( ToolTip )
@@ -698,9 +698,9 @@ void FDisplayedAttachedAssetInfo::GenerateWidgetForNameColumn( TSharedPtr< SHori
 		[
 			SNew( STextBlock )
 			.ColorAndOpacity( FLinearColor::Gray )
-			.Text(LOCTEXT( "AttachedAssetPreviewText", "[Preview Only]" ).ToString() )
+			.Text(LOCTEXT( "AttachedAssetPreviewText", "[Preview Only]" ) )
 			.Font(TextFont)
-			.ToolTipText( LOCTEXT( "AttachedAssetPreviewText_ToolTip", "Attached assets in Persona are preview only and do not carry through to the game." ).ToString() )
+			.ToolTipText( LOCTEXT( "AttachedAssetPreviewText_ToolTip", "Attached assets in Persona are preview only and do not carry through to the game." ) )
 		];
 }
 
@@ -712,7 +712,7 @@ TSharedRef< SWidget > FDisplayedAttachedAssetInfo::GenerateWidgetForDataColumn()
 		.HAlign( HAlign_Left )
 		[
 			SNew( SCheckBox )
-			.ToolTipText( LOCTEXT( "TranslationCheckBoxToolTip", "Click to toggle visibility of this asset" ).ToString() )
+			.ToolTipText( LOCTEXT( "TranslationCheckBoxToolTip", "Click to toggle visibility of this asset" ) )
 			.OnCheckStateChanged( this, &FDisplayedAttachedAssetInfo::OnToggleAssetDisplayed )
 			.IsChecked( this, &FDisplayedAttachedAssetInfo::IsAssetDisplayed )
 			.Style( FEditorStyle::Get(), "CheckboxLookToggleButtonCheckbox" )
@@ -821,7 +821,7 @@ void SSkeletonTree::Construct(const FArguments& InArgs)
 				SNew( SComboButton )
 				.ContentPadding(3)
 				.OnGetMenuContent( this, &SSkeletonTree::CreateBoneFilterMenu )
-				.ToolTipText( LOCTEXT( "BoneFilterToolTip", "Change which types of bones are shown" ).ToString() )
+				.ToolTipText( LOCTEXT( "BoneFilterToolTip", "Change which types of bones are shown" ) )
 				.ButtonContent()
 				[
 					SNew( STextBlock )
@@ -837,7 +837,7 @@ void SSkeletonTree::Construct(const FArguments& InArgs)
 				SNew( SComboButton )
 				.ContentPadding(3)
 				.OnGetMenuContent( this, &SSkeletonTree::CreateSocketFilterMenu )
-				.ToolTipText( LOCTEXT( "SocketFilterToolTip", "Change which types of sockets are shown" ).ToString() )
+				.ToolTipText( LOCTEXT( "SocketFilterToolTip", "Change which types of sockets are shown" ) )
 				.ButtonContent()
 				[
 					SNew( STextBlock )
@@ -863,11 +863,11 @@ void SSkeletonTree::Construct(const FArguments& InArgs)
 				SNew(SHeaderRow)
 
 				+ SHeaderRow::Column( ColumnID_BoneLabel )
-				.DefaultLabel( LOCTEXT( "SkeletonBoneNameLabel", "Name" ).ToString() )
+				.DefaultLabel( LOCTEXT( "SkeletonBoneNameLabel", "Name" ) )
 				.FillWidth( 0.75f )
 
 				+ SHeaderRow::Column( ColumnID_RetargetingLabel )
-				.DefaultLabel( LOCTEXT( "SkeletonBoneTranslationRetargetingLabel", "Translation Retargeting" ).ToString() )
+				.DefaultLabel( LOCTEXT( "SkeletonBoneTranslationRetargetingLabel", "Translation Retargeting" ) )
 				.FillWidth( 0.25f )
 			)
 		]
diff --git a/Engine/Source/Editor/Persona/Private/Shared/PersonaMode.cpp b/Engine/Source/Editor/Persona/Private/Shared/PersonaMode.cpp
--- a/Engine/Source/Editor/Persona/Private/Shared/PersonaMode.cpp
+++ b/Engine/Source/Editor/Persona/Private/Shared/PersonaMode.cpp
@@ -279,7 +279,7 @@ TSharedRef<SWidget> FAnimBlueprintDefaultsEditorSummoner::CreateTabBody(const FW
 						[
 							SNew( STextBlock )
 							.Font( FSlateFontInfo( FPaths::EngineContentDir() / TEXT("Slate/Fonts/Roboto-Bold.ttf"), 9 ) )
-							.Text( LOCTEXT("AnimBlueprintDefaultsPreviewMode", "Edit Preview").ToString() )
+							.Text( LOCTEXT("AnimBlueprintDefaultsPreviewMode", "Edit Preview") )
 						]
 					]
 				]
@@ -301,7 +301,7 @@ TSharedRef<SWidget> FAnimBlueprintDefaultsEditorSummoner::CreateTabBody(const FW
 						[
 							SNew( STextBlock )
 							.Font( FSlateFontInfo( FPaths::EngineContentDir() / TEXT("Slate/Fonts/Roboto-Bold.ttf"), 9 ) )
-							.Text( LOCTEXT("AnimBlueprintDefaultsDefaultsMode", "Edit Defaults").ToString() )
+							.Text( LOCTEXT("AnimBlueprintDefaultsDefaultsMode", "Edit Defaults") )
 						]
 					]
 				]
diff --git a/Engine/Source/Editor/PropertyEditor/Private/PropertyEditorToolkit.cpp b/Engine/Source/Editor/PropertyEditor/Private/PropertyEditorToolkit.cpp
--- a/Engine/Source/Editor/PropertyEditor/Private/PropertyEditorToolkit.cpp
+++ b/Engine/Source/Editor/PropertyEditor/Private/PropertyEditorToolkit.cpp
@@ -205,7 +205,7 @@ TSharedRef<SDockTab> FPropertyEditorToolkit::SpawnTab_PropertyTable( const FSpaw
 				[
 					SNew( STextBlock )
 					.Font( FEditorStyle::GetFontStyle( "PropertyEditor.AddColumnMessage.Font" ) )
-					.Text( LOCTEXT("GenericPropertiesTitle", "Pin Properties to Add Columns").ToString() )
+					.Text( LOCTEXT("GenericPropertiesTitle", "Pin Properties to Add Columns") )
 					.Visibility( this, &FPropertyEditorToolkit::GetAddColumnInstructionsOverlayVisibility )
 					.ColorAndOpacity( FEditorStyle::GetColor( "PropertyEditor.AddColumnMessage.ColorAndOpacity" ) )
 				]
diff --git a/Engine/Source/Editor/SceneOutliner/Private/SSceneOutliner.cpp b/Engine/Source/Editor/SceneOutliner/Private/SSceneOutliner.cpp
--- a/Engine/Source/Editor/SceneOutliner/Private/SSceneOutliner.cpp
+++ b/Engine/Source/Editor/SceneOutliner/Private/SSceneOutliner.cpp
@@ -288,7 +288,7 @@ namespace SceneOutliner
 			[
 				SNew( STextBlock )
 				.Visibility( this, &SSceneOutliner::GetEmptyLabelVisibility )
-				.Text( LOCTEXT( "EmptyLabel", "Empty" ).ToString() )
+				.Text( LOCTEXT( "EmptyLabel", "Empty" ) )
 				.ColorAndOpacity( FLinearColor( 0.4f, 1.0f, 0.4f ) )
 			]
 
@@ -378,7 +378,7 @@ namespace SceneOutliner
 						.Padding(2, 0, 0, 0)
 						.VAlign(VAlign_Center)
 						[
-							SNew(STextBlock).Text( LOCTEXT("ViewButton", "View Options").ToString() )
+							SNew(STextBlock).Text( LOCTEXT("ViewButton", "View Options") )
 						]
 					]
 				]
diff --git a/Engine/Source/Editor/SceneOutliner/Private/SceneOutlinerTreeItems.cpp b/Engine/Source/Editor/SceneOutliner/Private/SceneOutlinerTreeItems.cpp
--- a/Engine/Source/Editor/SceneOutliner/Private/SceneOutlinerTreeItems.cpp
+++ b/Engine/Source/Editor/SceneOutliner/Private/SceneOutlinerTreeItems.cpp
@@ -123,7 +123,7 @@ namespace SceneOutliner
 				auto TheFolder = Folders[0].Pin();
 				if (TheFolder.IsValid())
 				{
-					Operation->CurrentHoverText = FText::FromString(TheFolder->LeafName.ToString());
+					Operation->CurrentHoverText = FText::FromName(TheFolder->LeafName);
 					Operation->CurrentIconBrush = FEditorStyle::Get().GetBrush(TEXT("SceneOutliner.FolderClosed"));
 				}
 			}
diff --git a/Engine/Source/Editor/StaticMeshEditor/Private/StaticMeshEditorTools.cpp b/Engine/Source/Editor/StaticMeshEditor/Private/StaticMeshEditorTools.cpp
--- a/Engine/Source/Editor/StaticMeshEditor/Private/StaticMeshEditorTools.cpp
+++ b/Engine/Source/Editor/StaticMeshEditor/Private/StaticMeshEditorTools.cpp
@@ -98,7 +98,7 @@ void SConvexDecomposition::Construct(const FArguments& InArgs)
 			.VAlign(VAlign_Center)
 			[
 				SNew(STextBlock)
-				.Text( LOCTEXT("MaxNumHulls_ConvexDecomp", "Max Hulls").ToString() )
+				.Text( LOCTEXT("MaxNumHulls_ConvexDecomp", "Max Hulls") )
 			]
 
 			+SHorizontalBox::Slot()
@@ -123,7 +123,7 @@ void SConvexDecomposition::Construct(const FArguments& InArgs)
 			.VAlign(VAlign_Center)
 			[
 				SNew(STextBlock)
-				.Text( LOCTEXT("MaxHullVerts_ConvexDecomp", "Max Hull Verts").ToString() )
+				.Text( LOCTEXT("MaxHullVerts_ConvexDecomp", "Max Hull Verts") )
 			]
 
 			+SHorizontalBox::Slot()
@@ -148,7 +148,7 @@ void SConvexDecomposition::Construct(const FArguments& InArgs)
 			.Padding(8.0f, 0.0f, 8.0f, 0.0f)
 			[
 				SNew(SButton)
-				.Text( LOCTEXT("Apply_ConvexDecomp", "Apply").ToString() )
+				.Text( LOCTEXT("Apply_ConvexDecomp", "Apply") )
 				.OnClicked(this, &SConvexDecomposition::OnApplyDecomp)
 			]
 
@@ -157,7 +157,7 @@ void SConvexDecomposition::Construct(const FArguments& InArgs)
 			.Padding(8.0f, 0.0f, 8.0f, 0.0f)
 			[
 				SNew(SButton)
-				.Text( LOCTEXT("Defaults_ConvexDecomp", "Defaults").ToString() )
+				.Text( LOCTEXT("Defaults_ConvexDecomp", "Defaults") )
 				.OnClicked(this, &SConvexDecomposition::OnDefaults)
 			]
 		]
@@ -261,7 +261,7 @@ void FMeshBuildSettingsLayout::GenerateHeaderRowContent( FDetailWidgetRow& NodeR
 	NodeRow.NameContent()
 	[
 		SNew( STextBlock )
-		.Text( LOCTEXT("MeshBuildSettings", "Build Settings").ToString() )
+		.Text( LOCTEXT("MeshBuildSettings", "Build Settings") )
 		.Font( IDetailLayoutBuilder::GetDetailFont() )
 	];
 }
@@ -274,7 +274,7 @@ void FMeshBuildSettingsLayout::GenerateChildContent( IDetailChildrenBuilder& Chi
 		[
 			SNew(STextBlock)
 			.Font( IDetailLayoutBuilder::GetDetailFont() )
-			.Text(LOCTEXT("RecomputeNormals", "Recompute Normals").ToString())
+			.Text(LOCTEXT("RecomputeNormals", "Recompute Normals"))
 		
 		]
 		.ValueContent()
@@ -291,7 +291,7 @@ void FMeshBuildSettingsLayout::GenerateChildContent( IDetailChildrenBuilder& Chi
 		[
 			SNew(STextBlock)
 			.Font( IDetailLayoutBuilder::GetDetailFont() )
-			.Text(LOCTEXT("RecomputeTangents", "Recompute Tangents").ToString())
+			.Text(LOCTEXT("RecomputeTangents", "Recompute Tangents"))
 		]
 		.ValueContent()
 		[
@@ -307,7 +307,7 @@ void FMeshBuildSettingsLayout::GenerateChildContent( IDetailChildrenBuilder& Chi
 		[
 			SNew(STextBlock)
 			.Font( IDetailLayoutBuilder::GetDetailFont() )
-			.Text(LOCTEXT("RemoveDegenerates", "Remove Degenerates").ToString())
+			.Text(LOCTEXT("RemoveDegenerates", "Remove Degenerates"))
 		]
 		.ValueContent()
 		[
@@ -323,7 +323,7 @@ void FMeshBuildSettingsLayout::GenerateChildContent( IDetailChildrenBuilder& Chi
 		[
 			SNew(STextBlock)
 			.Font( IDetailLayoutBuilder::GetDetailFont() )
-			.Text(LOCTEXT("UseFullPrecisionUVs", "Use Full Precision UVs").ToString())
+			.Text(LOCTEXT("UseFullPrecisionUVs", "Use Full Precision UVs"))
 		]
 		.ValueContent()
 		[
@@ -339,7 +339,7 @@ void FMeshBuildSettingsLayout::GenerateChildContent( IDetailChildrenBuilder& Chi
 		[
 			SNew(STextBlock)
 			.Font(IDetailLayoutBuilder::GetDetailFont())
-			.Text(LOCTEXT("BuildScale", "Build Scale").ToString())
+			.Text(LOCTEXT("BuildScale", "Build Scale"))
 			.ToolTipText( LOCTEXT("BuildScale_ToolTip", "The local scale applied when building the mesh") )
 		]
 		.ValueContent()
@@ -409,7 +409,7 @@ void FMeshBuildSettingsLayout::GenerateChildContent( IDetailChildrenBuilder& Chi
 			.IsEnabled(ParentLODSettings.Pin().ToSharedRef(), &FLevelOfDetailSettingsLayout::IsApplyNeeded)
 			[
 				SNew( STextBlock )
-				.Text(LOCTEXT("ApplyChanges", "Apply Changes").ToString())
+				.Text(LOCTEXT("ApplyChanges", "Apply Changes"))
 				.Font( IDetailLayoutBuilder::GetDetailFont() )
 			]
 		];
@@ -515,7 +515,7 @@ void FMeshReductionSettingsLayout::GenerateHeaderRowContent( FDetailWidgetRow& N
 	NodeRow.NameContent()
 	[
 		SNew( STextBlock )
-		.Text( LOCTEXT("MeshReductionSettings", "Reduction Settings").ToString() )
+		.Text( LOCTEXT("MeshReductionSettings", "Reduction Settings") )
 		.Font( IDetailLayoutBuilder::GetDetailFont() )
 	];
 }
@@ -530,7 +530,7 @@ void FMeshReductionSettingsLayout::GenerateChildContent( IDetailChildrenBuilder&
 		[
 			SNew(STextBlock)
 			.Font( IDetailLayoutBuilder::GetDetailFont() )
-			.Text(LOCTEXT("PercentTriangles", "Percent Triangles").ToString())
+			.Text(LOCTEXT("PercentTriangles", "Percent Triangles"))
 		]
 		.ValueContent()
 		[
@@ -550,7 +550,7 @@ void FMeshReductionSettingsLayout::GenerateChildContent( IDetailChildrenBuilder&
 		[
 			SNew(STextBlock)
 			.Font( IDetailLayoutBuilder::GetDetailFont() )
-			.Text(LOCTEXT("MaxDeviation", "Max Deviation").ToString())
+			.Text(LOCTEXT("MaxDeviation", "Max Deviation"))
 		]
 		.ValueContent()
 		[
@@ -570,7 +570,7 @@ void FMeshReductionSettingsLayout::GenerateChildContent( IDetailChildrenBuilder&
 		[
 			SNew( STextBlock )
 			.Font( IDetailLayoutBuilder::GetDetailFont() )
-			.Text( LOCTEXT("Silhouette_MeshSimplification", "Silhouette").ToString() )
+			.Text( LOCTEXT("Silhouette_MeshSimplification", "Silhouette") )
 		]
 		.ValueContent()
 		[
@@ -590,7 +590,7 @@ void FMeshReductionSettingsLayout::GenerateChildContent( IDetailChildrenBuilder&
 		[
 			SNew( STextBlock )
 			.Font( IDetailLayoutBuilder::GetDetailFont() )
-			.Text( LOCTEXT("Texture_MeshSimplification", "Texture").ToString() )
+			.Text( LOCTEXT("Texture_MeshSimplification", "Texture") )
 		]
 		.ValueContent()
 		[
@@ -610,7 +610,7 @@ void FMeshReductionSettingsLayout::GenerateChildContent( IDetailChildrenBuilder&
 		[
 			SNew( STextBlock )
 			.Font( IDetailLayoutBuilder::GetDetailFont() )
-			.Text( LOCTEXT("Shading_MeshSimplification", "Shading").ToString() )
+			.Text( LOCTEXT("Shading_MeshSimplification", "Shading") )
 		]
 		.ValueContent()
 		[
@@ -630,7 +630,7 @@ void FMeshReductionSettingsLayout::GenerateChildContent( IDetailChildrenBuilder&
 		[
 			SNew(STextBlock)
 			.Font( IDetailLayoutBuilder::GetDetailFont() )
-			.Text(LOCTEXT("WeldingThreshold", "Welding Threshold").ToString())
+			.Text(LOCTEXT("WeldingThreshold", "Welding Threshold"))
 		]
 		.ValueContent()
 		[
@@ -649,7 +649,7 @@ void FMeshReductionSettingsLayout::GenerateChildContent( IDetailChildrenBuilder&
 		[
 			SNew(STextBlock)
 			.Font( IDetailLayoutBuilder::GetDetailFont() )
-			.Text(LOCTEXT("RecomputeNormals", "Recompute Normals").ToString())
+			.Text(LOCTEXT("RecomputeNormals", "Recompute Normals"))
 
 		]
 		.ValueContent()
@@ -666,7 +666,7 @@ void FMeshReductionSettingsLayout::GenerateChildContent( IDetailChildrenBuilder&
 		[
 			SNew(STextBlock)
 			.Font( IDetailLayoutBuilder::GetDetailFont() )
-			.Text(LOCTEXT("HardEdgeAngle", "Hard Edge Angle").ToString())
+			.Text(LOCTEXT("HardEdgeAngle", "Hard Edge Angle"))
 		]
 		.ValueContent()
 		[
@@ -690,7 +690,7 @@ void FMeshReductionSettingsLayout::GenerateChildContent( IDetailChildrenBuilder&
 				.IsEnabled(ParentLODSettings.Pin().ToSharedRef(), &FLevelOfDetailSettingsLayout::IsApplyNeeded)
 				[
 					SNew( STextBlock )
-					.Text(LOCTEXT("ApplyChanges", "Apply Changes").ToString())
+					.Text(LOCTEXT("ApplyChanges", "Apply Changes"))
 					.Font( IDetailLayoutBuilder::GetDetailFont() )
 				]
 			];
@@ -866,12 +866,12 @@ TSharedRef<SWidget> FMeshSectionSettingsLayout::OnGenerateNameWidgetsForMaterial
 		SNew(SCheckBox)
 		.IsChecked(this, &FMeshSectionSettingsLayout::IsSectionSelected, SlotIndex)
 		.OnCheckStateChanged(this, &FMeshSectionSettingsLayout::OnSectionSelectedChanged, SlotIndex)
-		.ToolTipText(LOCTEXT("Highlight_ToolTip", "Highlights this section in the viewport").ToString())
+		.ToolTipText(LOCTEXT("Highlight_ToolTip", "Highlights this section in the viewport"))
 		[
 			SNew(STextBlock)
 			.Font(IDetailLayoutBuilder::GetDetailFont())
 			.ColorAndOpacity( FLinearColor( 0.4f, 0.4f, 0.4f, 1.0f) )
-			.Text(LOCTEXT("Highlight", "Highlight").ToString())
+			.Text(LOCTEXT("Highlight", "Highlight"))
 
 		];
 }
@@ -888,7 +888,7 @@ TSharedRef<SWidget> FMeshSectionSettingsLayout::OnGenerateWidgetsForMaterial(UMa
 			[
 				SNew(STextBlock)
 					.Font(FEditorStyle::GetFontStyle("StaticMeshEditor.NormalFont"))
-					.Text(LOCTEXT("CastShadow", "Cast Shadow").ToString())
+					.Text(LOCTEXT("CastShadow", "Cast Shadow"))
 			]
 		]
 		+SVerticalBox::Slot()
@@ -901,7 +901,7 @@ TSharedRef<SWidget> FMeshSectionSettingsLayout::OnGenerateWidgetsForMaterial(UMa
 			[
 				SNew(STextBlock)
 					.Font(FEditorStyle::GetFontStyle("StaticMeshEditor.NormalFont"))
-					.Text(LOCTEXT("EnableCollision", "Enable Collision").ToString())
+					.Text(LOCTEXT("EnableCollision", "Enable Collision"))
 			]
 		];
 }
@@ -1043,7 +1043,7 @@ void FLevelOfDetailSettingsLayout::AddToDetailsPanel( IDetailLayoutBuilder& Deta
 	[
 		SNew(STextBlock)
 		.Font( IDetailLayoutBuilder::GetDetailFont() )
-		.Text(LOCTEXT("LODGroup", "LOD Group").ToString())
+		.Text(LOCTEXT("LODGroup", "LOD Group"))
 	]
 	.ValueContent()
 	[
@@ -1059,7 +1059,7 @@ void FLevelOfDetailSettingsLayout::AddToDetailsPanel( IDetailLayoutBuilder& Deta
 		[
 			SNew(STextBlock)
 			.Font( IDetailLayoutBuilder::GetDetailFont() )
-			.Text(LOCTEXT("LODImport", "LOD Import").ToString())
+			.Text(LOCTEXT("LODImport", "LOD Import"))
 		]
 	.ValueContent()
 		[
@@ -1077,7 +1077,7 @@ void FLevelOfDetailSettingsLayout::AddToDetailsPanel( IDetailLayoutBuilder& Deta
 	[
 		SNew(STextBlock)
 		.Font( IDetailLayoutBuilder::GetDetailFont() )
-		.Text(LOCTEXT("NumberOfLODs", "Number of LODs").ToString())
+		.Text(LOCTEXT("NumberOfLODs", "Number of LODs"))
 	]
 	.ValueContent()
 	[
@@ -1098,7 +1098,7 @@ void FLevelOfDetailSettingsLayout::AddToDetailsPanel( IDetailLayoutBuilder& Deta
 	[
 		SNew(STextBlock)
 		.Font( IDetailLayoutBuilder::GetDetailFont() )
-		.Text(LOCTEXT("AutoComputeLOD", "Auto Compute LOD Distances").ToString())
+		.Text(LOCTEXT("AutoComputeLOD", "Auto Compute LOD Distances"))
 	]
 	.ValueContent()
 	[
@@ -1116,7 +1116,7 @@ void FLevelOfDetailSettingsLayout::AddToDetailsPanel( IDetailLayoutBuilder& Deta
 		.IsEnabled(this, &FLevelOfDetailSettingsLayout::IsApplyNeeded)
 		[
 			SNew( STextBlock )
-			.Text(LOCTEXT("ApplyChanges", "Apply Changes").ToString())
+			.Text(LOCTEXT("ApplyChanges", "Apply Changes"))
 			.Font( DetailBuilder.GetDetailFont() )
 		]
 	];
@@ -1128,7 +1128,7 @@ void FLevelOfDetailSettingsLayout::AddToDetailsPanel( IDetailLayoutBuilder& Deta
 	[
 		SNew(STextBlock)
 		.Font(FEditorStyle::GetFontStyle("StaticMeshEditor.NormalFont"))
-		.Text(LOCTEXT("AllowedPixelError", "Auto Distance Error").ToString())
+		.Text(LOCTEXT("AllowedPixelError", "Auto Distance Error"))
 	]
 	.ValueContent()
 	[
@@ -1201,7 +1201,7 @@ void FLevelOfDetailSettingsLayout::AddLODLevelCategories( IDetailLayoutBuilder&
 			FString CategoryName = FString(TEXT("LOD"));
 			CategoryName.AppendInt( LODIndex );
 
-			FString LODLevelString = FText::Format( LOCTEXT("LODLevel", "LOD {0}"), FText::AsNumber( LODIndex ) ).ToString();
+			FString LODLevelString = FText::Format( LOCTEXT("LODLevel", "LOD{0}"), FText::AsNumber( LODIndex ) ).ToString();
 
 			IDetailCategoryBuilder& LODCategory = DetailBuilder.EditCategory( *CategoryName, LODLevelString, ECategoryPriority::Important );
 
@@ -1250,7 +1250,7 @@ void FLevelOfDetailSettingsLayout::AddLODLevelCategories( IDetailLayoutBuilder&
 				[
 					SNew(STextBlock)
 					.Font( IDetailLayoutBuilder::GetDetailFont() )
-					.Text(LOCTEXT("Distance", "Distance").ToString())
+					.Text(LOCTEXT("Distance", "Distance"))
 				]
 				.ValueContent()
 				[
@@ -1657,7 +1657,7 @@ void SGenerateUniqueUVs::Construct(const FArguments& InArgs)
 			.AutoHeight()
 			[
 				SNew(STextBlock)
-					.Text( LOCTEXT("CreationMode", "Creation Mode").ToString() )
+					.Text( LOCTEXT("CreationMode", "Creation Mode") )
 			]
 
 			+SVerticalBox::Slot()
@@ -1676,7 +1676,7 @@ void SGenerateUniqueUVs::Construct(const FArguments& InArgs)
 						.OnCheckStateChanged( this, &SGenerateUniqueUVs::OnCreationModeChanged, CreateNew )
 						[
 							SNew(STextBlock)
-								.Text( LOCTEXT("CreateNew", "Create New").ToString() )
+								.Text( LOCTEXT("CreateNew", "Create New") )
 						]
 				]
 
@@ -1690,7 +1690,7 @@ void SGenerateUniqueUVs::Construct(const FArguments& InArgs)
 						.OnCheckStateChanged( this, &SGenerateUniqueUVs::OnCreationModeChanged, UseChannel0 )
 						[
 							SNew(STextBlock)
-								.Text( LOCTEXT("LayoutUsing0Channel", "Layout using 0 channel").ToString() )
+								.Text( LOCTEXT("LayoutUsing0Channel", "Layout using 0 channel") )
 						]
 				]
 			]
@@ -1705,7 +1705,7 @@ void SGenerateUniqueUVs::Construct(const FArguments& InArgs)
 					.FillWidth(3.0f)
 				[
 					SNew(STextBlock)
-						.Text( LOCTEXT("UVChannelSaveSelection", "UV channel to save results to:").ToString() )
+						.Text( LOCTEXT("UVChannelSaveSelection", "UV channel to save results to:") )
 				]
 
 				+SHorizontalBox::Slot()
@@ -1724,7 +1724,7 @@ void SGenerateUniqueUVs::Construct(const FArguments& InArgs)
 				.Padding( 0.0f, 8.0f, 0.0f, 0.0f)
 			[
 				SNew(STextBlock)
-					.Text( LOCTEXT("SelectGenerationMethod", "Select Generation Method").ToString() )
+					.Text( LOCTEXT("SelectGenerationMethod", "Select Generation Method") )
 					.IsEnabled(this, &SGenerateUniqueUVs::IsCreateNew)
 			]
 
@@ -1744,7 +1744,7 @@ void SGenerateUniqueUVs::Construct(const FArguments& InArgs)
 						.IsEnabled(this, &SGenerateUniqueUVs::IsCreateNew)
 					[
 						SNew(STextBlock)
-							.Text( LOCTEXT("LimitMaxStretching", "Limit maximum stretching").ToString() )
+							.Text( LOCTEXT("LimitMaxStretching", "Limit maximum stretching") )
 					]
 				]
 
@@ -1778,7 +1778,7 @@ void SGenerateUniqueUVs::Construct(const FArguments& InArgs)
 						.IsEnabled(this, &SGenerateUniqueUVs::IsCreateNew)
 					[
 						SNew(STextBlock)
-							.Text( LOCTEXT("LimitMaxCharts", "Limit maximum number of charts").ToString() )
+							.Text( LOCTEXT("LimitMaxCharts", "Limit maximum number of charts") )
 					]
 				]
 
@@ -1806,7 +1806,7 @@ void SGenerateUniqueUVs::Construct(const FArguments& InArgs)
 					.FillWidth(3.0f)
 				[
 					SNew(STextBlock)
-						.Text( LOCTEXT("LimitChartSpacing", "Limit spacing between charts:").ToString() )
+						.Text( LOCTEXT("LimitChartSpacing", "Limit spacing between charts:") )
 						.IsEnabled(this, &SGenerateUniqueUVs::IsCreateNew)
 				]
 
@@ -1830,7 +1830,7 @@ void SGenerateUniqueUVs::Construct(const FArguments& InArgs)
 				.Padding( 0.0f, 8.0f, 0.0f, 0.0f)
 			[
 				SNew(SButton)
-					.Text( LOCTEXT("Apply", "Apply").ToString() )
+					.Text( LOCTEXT("Apply", "Apply") )
 					.OnClicked( this, &SGenerateUniqueUVs::OnApply )
 			]
 		];
diff --git a/Engine/Source/Editor/TaskBrowser/Private/STaskBrowser.cpp b/Engine/Source/Editor/TaskBrowser/Private/STaskBrowser.cpp
--- a/Engine/Source/Editor/TaskBrowser/Private/STaskBrowser.cpp
+++ b/Engine/Source/Editor/TaskBrowser/Private/STaskBrowser.cpp
@@ -562,7 +562,7 @@ void STaskBrowser::Construct(const FArguments& InArgs)
 						.IsEnabled_Raw( this, &STaskBrowser::GetMarkCompleteEnabled )
 						.OnClicked( this, &STaskBrowser::OnMarkCompleteClicked )
 						[
-							SNew(STextBlock) .Text( LOCTEXT("MarkComplete", "Mark Complete...").ToString() )
+							SNew(STextBlock) .Text( LOCTEXT("MarkComplete", "Mark Complete...") )
 						]
 					]
 					+SHorizontalBox::Slot()
@@ -575,7 +575,7 @@ void STaskBrowser::Construct(const FArguments& InArgs)
 						.VAlign(VAlign_Center)
 						.OnClicked( this, &STaskBrowser::OnRefreshViewClicked )
 						[
-							SNew(STextBlock) .Text( LOCTEXT("RefreshView", "Refresh View").ToString() )
+							SNew(STextBlock) .Text( LOCTEXT("RefreshView", "Refresh View") )
 						]
 					]
 					+SHorizontalBox::Slot()
@@ -584,7 +584,7 @@ void STaskBrowser::Construct(const FArguments& InArgs)
 					.HAlign(HAlign_Center)
 					.VAlign(VAlign_Center)
 					[
-						SNew(STextBlock) .Text( LOCTEXT("DatabaseFilter", "Database filter").ToString() )
+						SNew(STextBlock) .Text( LOCTEXT("DatabaseFilter", "Database filter") )
 					]
 					+SHorizontalBox::Slot()
 					.AutoWidth()
@@ -604,7 +604,7 @@ void STaskBrowser::Construct(const FArguments& InArgs)
 					.AutoWidth()
 					.Padding( SmallPadding )
 					[
-						SNew(STextBlock) .Text( LOCTEXT("Display", "Display").ToString() )
+						SNew(STextBlock) .Text( LOCTEXT("Display", "Display") )
 					]
 					+SHorizontalBox::Slot()
 					.AutoWidth()
@@ -616,7 +616,7 @@ void STaskBrowser::Construct(const FArguments& InArgs)
 					.AutoWidth()
 					.Padding( SmallPadding )
 					[
-						SNew(STextBlock) .Text( LOCTEXT("OpenOnly", "Open only").ToString() )
+						SNew(STextBlock) .Text( LOCTEXT("OpenOnly", "Open only") )
 					]
 					+SHorizontalBox::Slot()
 					.AutoWidth()
@@ -628,7 +628,7 @@ void STaskBrowser::Construct(const FArguments& InArgs)
 					.AutoWidth()
 					.Padding( SmallPadding )
 					[
-						SNew(STextBlock) .Text( LOCTEXT("AssignedToMe", "Assigned to me").ToString() )
+						SNew(STextBlock) .Text( LOCTEXT("AssignedToMe", "Assigned to me") )
 					]
 					+SHorizontalBox::Slot()
 					.AutoWidth()
@@ -640,7 +640,7 @@ void STaskBrowser::Construct(const FArguments& InArgs)
 					.AutoWidth()
 					.Padding( SmallPadding )
 					[
-						SNew(STextBlock) .Text( LOCTEXT("CreatedByMe", "Created by me").ToString() )
+						SNew(STextBlock) .Text( LOCTEXT("CreatedByMe", "Created by me") )
 					]
 					+SHorizontalBox::Slot()
 					.AutoWidth()
@@ -652,7 +652,7 @@ void STaskBrowser::Construct(const FArguments& InArgs)
 					.AutoWidth()
 					.Padding( SmallPadding )
 					[
-						SNew(STextBlock) .Text( LOCTEXT("CurrentMap", "Current map").ToString() )
+						SNew(STextBlock) .Text( LOCTEXT("CurrentMap", "Current map") )
 					]
 				]
 				+SVerticalBox::Slot()
@@ -770,7 +770,7 @@ void STaskBrowser::Construct(const FArguments& InArgs)
 				.VAlign(VAlign_Center)
 				.OnClicked( this, &STaskBrowser::OnSettingsClicked )
 				[
-					SNew(STextBlock) .Text( LOCTEXT("Settings", "Settings...").ToString() )
+					SNew(STextBlock) .Text( LOCTEXT("Settings", "Settings...") )
 				]
 			]
 			+SHorizontalBox::Slot()
@@ -1043,7 +1043,7 @@ FReply STaskBrowser::OnConnectClicked()
 	return FReply::Handled();
 }
 
-FString STaskBrowser::GetConnectText() const
+FText STaskBrowser::GetConnectText() const
 {
 	switch( TaskDataManager->GetConnectionStatus() )
 	{
@@ -1051,9 +1051,9 @@ FString STaskBrowser::GetConnectText() const
 	case ETaskDataManagerStatus::ReadyToConnect:
 	case ETaskDataManagerStatus::ConnectionFailed:
 	case ETaskDataManagerStatus::Disconnecting:
-		return NSLOCTEXT("UnrealEd", "TaskBrowser_ConnectButtonLabel", "Connect" ).ToString();
+		return NSLOCTEXT("UnrealEd", "TaskBrowser_ConnectButtonLabel", "Connect" );
 	default:
-		return NSLOCTEXT("UnrealEd", "TaskBrowser_DisconnectButtonLabel", "Disconnect" ).ToString();
+		return NSLOCTEXT("UnrealEd", "TaskBrowser_DisconnectButtonLabel", "Disconnect" );
 	}
 }
 
@@ -1064,30 +1064,30 @@ FReply STaskBrowser::OnSettingsClicked()
 	return FReply::Handled();
 }
 
-FString STaskBrowser::GetStatusText() const
+FText STaskBrowser::GetStatusText() const
 {
 	switch( TaskDataManager->GetGUIStatus() )
 	{
 	case ETaskDataManagerStatus::ReadyToConnect:
-		return NSLOCTEXT("UnrealEd", "TaskBrowser_ServerStatus_ReadyToConnect", "Not connected.  Click the Connect button to login to the task database." ).ToString();
+		return NSLOCTEXT("UnrealEd", "TaskBrowser_ServerStatus_ReadyToConnect", "Not connected.  Click the Connect button to login to the task database." );
 	case ETaskDataManagerStatus::Connecting:
-		return NSLOCTEXT("UnrealEd", "TaskBrowser_ServerStatus_Connecting", "Connecting to server and logging in..." ).ToString();
+		return NSLOCTEXT("UnrealEd", "TaskBrowser_ServerStatus_Connecting", "Connecting to server and logging in..." );
 	case ETaskDataManagerStatus::Connected:
-		return NSLOCTEXT("UnrealEd", "TaskBrowser_ServerStatus_Connected", "Connected to the task database." ).ToString();
+		return NSLOCTEXT("UnrealEd", "TaskBrowser_ServerStatus_Connected", "Connected to the task database." );
 	case ETaskDataManagerStatus::ConnectionFailed:
-		return NSLOCTEXT("UnrealEd", "TaskBrowser_ServerStatus_ConnectionFailed", "Failed to connect to the task database.  Click the Connect button to try again." ).ToString();
+		return NSLOCTEXT("UnrealEd", "TaskBrowser_ServerStatus_ConnectionFailed", "Failed to connect to the task database.  Click the Connect button to try again." );
 	case ETaskDataManagerStatus::Disconnecting:
-		return NSLOCTEXT("UnrealEd", "TaskBrowser_ServerStatus_Disconnecting", "Disconnecting..." ).ToString();
+		return NSLOCTEXT("UnrealEd", "TaskBrowser_ServerStatus_Disconnecting", "Disconnecting..." );
 	case ETaskDataManagerStatus::QueryingFilters:
-		return NSLOCTEXT("UnrealEd", "TaskBrowser_ServerStatus_QueryingFilters", "Retrieving filters..." ).ToString();
+		return NSLOCTEXT("UnrealEd", "TaskBrowser_ServerStatus_QueryingFilters", "Retrieving filters..." );
 	case ETaskDataManagerStatus::QueryingTasks:
-		return NSLOCTEXT("UnrealEd", "TaskBrowser_ServerStatus_QueryingTasks", "Retrieving task list from server..." ).ToString();
+		return NSLOCTEXT("UnrealEd", "TaskBrowser_ServerStatus_QueryingTasks", "Retrieving task list from server..." );
 	case ETaskDataManagerStatus::QueryingTaskDetails:
-		return NSLOCTEXT("UnrealEd", "TaskBrowser_ServerStatus_QueryingTaskDetails", "Retrieving task information..." ).ToString();
+		return NSLOCTEXT("UnrealEd", "TaskBrowser_ServerStatus_QueryingTaskDetails", "Retrieving task information..." );
 	case ETaskDataManagerStatus::MarkingTaskComplete:
-		return NSLOCTEXT("UnrealEd", "TaskBrowser_ServerStatus_MarkingTaskComplete", "Marking task as fixed..." ).ToString();
+		return NSLOCTEXT("UnrealEd", "TaskBrowser_ServerStatus_MarkingTaskComplete", "Marking task as fixed..." );
 	default:
-		return NSLOCTEXT("UnrealEd", "TaskBrowser_ServerStatus_FailedToInit", "Failed to initialize.  No task database providers are available." ).ToString();
+		return NSLOCTEXT("UnrealEd", "TaskBrowser_ServerStatus_FailedToInit", "Failed to initialize.  No task database providers are available." );
 	}
 }
 
diff --git a/Engine/Source/Editor/TaskBrowser/Private/STaskBrowser.h b/Engine/Source/Editor/TaskBrowser/Private/STaskBrowser.h
--- a/Engine/Source/Editor/TaskBrowser/Private/STaskBrowser.h
+++ b/Engine/Source/Editor/TaskBrowser/Private/STaskBrowser.h
@@ -196,13 +196,13 @@ private:
 	FReply OnConnectClicked();
 
 	/** Return the connect button text based on the state*/
-	FString GetConnectText() const;
+	FText GetConnectText() const;
 
 	/** Handle when the Refresh View button is pressed */
 	FReply OnSettingsClicked();
 
 	/** Return the status text based on the state*/
-	FString GetStatusText() const;
+	FText GetStatusText() const;
 
 	/** Spawn our task complete modal dialog */
 	void CompleteTaskDialog( TArray< uint32 >& TaskNumbersToFix );
diff --git a/Engine/Source/Editor/TaskBrowser/Private/STaskComplete.cpp b/Engine/Source/Editor/TaskBrowser/Private/STaskComplete.cpp
--- a/Engine/Source/Editor/TaskBrowser/Private/STaskComplete.cpp
+++ b/Engine/Source/Editor/TaskBrowser/Private/STaskComplete.cpp
@@ -61,7 +61,7 @@ void STaskComplete::Construct(const FArguments& InArgs)
 					.Padding( SmallPadding )
 					.FillWidth( 0.18f )
 					[
-						SNew(STextBlock) .Text( LOCTEXT("Resolution", "Resolution").ToString() )
+						SNew(STextBlock) .Text( LOCTEXT("Resolution", "Resolution") )
 					]
 					+SHorizontalBox::Slot()
 					.FillWidth( 0.33f )
@@ -83,7 +83,7 @@ void STaskComplete::Construct(const FArguments& InArgs)
 					.Padding( SmallPadding )
 					.FillWidth( 0.18f )
 					[
-						SNew(STextBlock) .Text( LOCTEXT("Comments", "Comments").ToString() )
+						SNew(STextBlock) .Text( LOCTEXT("Comments", "Comments") )
 					]
 					+SHorizontalBox::Slot()
 					.FillWidth( 0.82f )
@@ -102,7 +102,7 @@ void STaskComplete::Construct(const FArguments& InArgs)
 					.Padding( SmallPadding )
 					.FillWidth( 0.184f )
 					[
-						SNew(STextBlock) .Text( LOCTEXT("Changelist", "Changelist").ToString() )
+						SNew(STextBlock) .Text( LOCTEXT("Changelist", "Changelist") )
 					]
 					+SHorizontalBox::Slot()
 					.FillWidth( 0.17f )
@@ -115,7 +115,7 @@ void STaskComplete::Construct(const FArguments& InArgs)
 					.Padding( SmallPadding )
 					.FillWidth( 0.408f )
 					[
-						SNew(STextBlock) .Text( LOCTEXT("TimeToComplete", "Time to complete (hours)").ToString() )
+						SNew(STextBlock) .Text( LOCTEXT("TimeToComplete", "Time to complete (hours)") )
 					]
 					+SHorizontalBox::Slot()
 					.FillWidth( 0.08f )
@@ -141,7 +141,7 @@ void STaskComplete::Construct(const FArguments& InArgs)
 						.HAlign(HAlign_Center)
 						.OnClicked( this, &STaskComplete::OnOKClicked )
 						[
-							SNew(STextBlock) .Text( LOCTEXT("OK", "OK").ToString() )
+							SNew(STextBlock) .Text( LOCTEXT("OK", "OK") )
 						]
 					]
 					+SUniformGridPanel::Slot(1,0)
@@ -151,7 +151,7 @@ void STaskComplete::Construct(const FArguments& InArgs)
 						.HAlign(HAlign_Center)
 						.OnClicked( this, &STaskComplete::OnCancelClicked )
 						[
-							SNew(STextBlock) .Text( LOCTEXT("Cancel", "Cancel").ToString() )
+							SNew(STextBlock) .Text( LOCTEXT("Cancel", "Cancel") )
 						]
 					]
 				]
diff --git a/Engine/Source/Editor/TaskBrowser/Private/STaskSettings.cpp b/Engine/Source/Editor/TaskBrowser/Private/STaskSettings.cpp
--- a/Engine/Source/Editor/TaskBrowser/Private/STaskSettings.cpp
+++ b/Engine/Source/Editor/TaskBrowser/Private/STaskSettings.cpp
@@ -45,7 +45,7 @@ void STaskSettings::Construct(const FArguments& InArgs)
 				.AutoHeight()
 				.Padding( HeaderPadding )
 				[
-					SNew(STextBlock) .Text( LOCTEXT("Connection", "Connection").ToString() )
+					SNew(STextBlock) .Text( LOCTEXT("Connection", "Connection") )
 				]
 				+SVerticalBox::Slot()
 				.AutoHeight()
@@ -64,7 +64,7 @@ void STaskSettings::Construct(const FArguments& InArgs)
 							.Padding( SmallPadding )
 							.FillWidth( SettingWidth )
 							[
-								SNew(STextBlock) .Text( LOCTEXT("Server", "Server").ToString() )
+								SNew(STextBlock) .Text( LOCTEXT("Server", "Server") )
 							]
 							+SHorizontalBox::Slot()
 							.FillWidth( 0.5f )
@@ -77,7 +77,7 @@ void STaskSettings::Construct(const FArguments& InArgs)
 							.Padding( SmallPadding )
 							.FillWidth( 0.13f )
 							[
-								SNew(STextBlock) .Text( LOCTEXT("Port", "Port").ToString() )
+								SNew(STextBlock) .Text( LOCTEXT("Port", "Port") )
 							]
 							+SHorizontalBox::Slot()
 							.FillWidth( 0.10f )
@@ -96,7 +96,7 @@ void STaskSettings::Construct(const FArguments& InArgs)
 							.Padding( SmallPadding )
 							.FillWidth( SettingWidth )
 							[
-								SNew(STextBlock) .Text( LOCTEXT("Login", "Login").ToString() )
+								SNew(STextBlock) .Text( LOCTEXT("Login", "Login") )
 							]
 							+SHorizontalBox::Slot()
 							.MaxWidth( FieldWidth )
@@ -116,7 +116,7 @@ void STaskSettings::Construct(const FArguments& InArgs)
 							.Padding( SmallPadding )
 							.FillWidth( SettingWidth )
 							[
-								SNew(STextBlock) .Text( LOCTEXT("Password", "Password").ToString() )
+								SNew(STextBlock) .Text( LOCTEXT("Password", "Password") )
 							]
 							+SHorizontalBox::Slot()
 							.MaxWidth( FieldWidth )
@@ -136,7 +136,7 @@ void STaskSettings::Construct(const FArguments& InArgs)
 							.Padding( SmallPadding )
 							.FillWidth( SettingWidth )
 							[
-								SNew(STextBlock) .Text( LOCTEXT("ProjectName", "Project name").ToString() )
+								SNew(STextBlock) .Text( LOCTEXT("ProjectName", "Project name") )
 							]
 							+SHorizontalBox::Slot()
 							.MaxWidth( FieldWidth )
@@ -151,7 +151,7 @@ void STaskSettings::Construct(const FArguments& InArgs)
 				.AutoHeight()
 				.Padding( HeaderPadding )
 				[
-					SNew(STextBlock) .Text( LOCTEXT("Preferences", "Preferences").ToString() )
+					SNew(STextBlock) .Text( LOCTEXT("Preferences", "Preferences") )
 				]
 				+SVerticalBox::Slot()
 				.AutoHeight()
@@ -175,7 +175,7 @@ void STaskSettings::Construct(const FArguments& InArgs)
 							.HAlign(HAlign_Right)
 							.VAlign(VAlign_Center)
 							[
-								SNew(STextBlock) .Text( LOCTEXT("AutoConnectAtStartup", "Auto connect at startup").ToString() )
+								SNew(STextBlock) .Text( LOCTEXT("AutoConnectAtStartup", "Auto connect at startup") )
 							]
 							+SHorizontalBox::Slot()	// Set this box to be as wide as possible, forcing the rest in this row to be right aligned
 							.FillWidth( 1.0f )
@@ -191,7 +191,7 @@ void STaskSettings::Construct(const FArguments& InArgs)
 							.HAlign(HAlign_Right)
 							.VAlign(VAlign_Center)
 							[
-								SNew(STextBlock) .Text( LOCTEXT("UseSingleSignOn", "Use single sign-on").ToString() )
+								SNew(STextBlock) .Text( LOCTEXT("UseSingleSignOn", "Use single sign-on") )
 							]
 						]
 					]
@@ -215,7 +215,7 @@ void STaskSettings::Construct(const FArguments& InArgs)
 						.HAlign(HAlign_Center)
 						.OnClicked( this, &STaskSettings::OnOKClicked )
 						[
-							SNew(STextBlock) .Text( LOCTEXT("OK", "OK").ToString() )
+							SNew(STextBlock) .Text( LOCTEXT("OK", "OK") )
 						]
 					]
 					+SUniformGridPanel::Slot(1,0)
@@ -225,7 +225,7 @@ void STaskSettings::Construct(const FArguments& InArgs)
 						.HAlign(HAlign_Center)
 						.OnClicked( this, &STaskSettings::OnCancelClicked )
 						[
-							SNew(STextBlock) .Text( LOCTEXT("Cancel", "Cancel").ToString() )
+							SNew(STextBlock) .Text( LOCTEXT("Cancel", "Cancel") )
 						]
 					]
 				]
diff --git a/Engine/Source/Editor/UnrealEd/Private/AnimSetViewer/LODSettingsWindow.cpp b/Engine/Source/Editor/UnrealEd/Private/AnimSetViewer/LODSettingsWindow.cpp
--- a/Engine/Source/Editor/UnrealEd/Private/AnimSetViewer/LODSettingsWindow.cpp
+++ b/Engine/Source/Editor/UnrealEd/Private/AnimSetViewer/LODSettingsWindow.cpp
@@ -310,7 +310,7 @@ private:
 				.VAlign(VAlign_Center)
 				[
 					SNew(STextBlock)
-					.Text(LOCTEXT("NumberOfLODs", "Number of LODs").ToString())
+					.Text(LOCTEXT("NumberOfLODs", "Number of LODs"))
 				]
 
 				+ SHorizontalBox::Slot()
@@ -339,8 +339,8 @@ private:
 												.DefaultQuality(DefaultQuality)
 												.Skeleton(UpdateContext.SkeletalMesh->Skeleton)
 												.DesiredLOD(LODLevel);
+			const FText AreaTitle = FText::Format(LOCTEXT("LODLevel", "LOD {0}"), FText::AsNumber(LODLevel));
 			
-			const FString AreaTitle = FText::Format( LOCTEXT("LODLevel", "LOD {0}"), FText::AsNumber( LODLevel ) ).ToString();
 			LODChainContent->AddSlot()
 				.AutoHeight()
 				.Padding(5.0)
@@ -376,7 +376,7 @@ private:
 			.AutoHeight()
 			[
 				SNew(SExpandableArea)
-				.AreaTitle(LOCTEXT("GenerateLODChain", "Generate LOD Chain").ToString())
+				.AreaTitle(LOCTEXT("GenerateLODChain", "Generate LOD Chain"))
 				.InitiallyCollapsed(false)
 				.BodyContent()
 				[
@@ -421,9 +421,9 @@ private:
 	}
 
 	/** Returns text for the 'Generate' Button. */
-	FString GetGenerateButtonText() const
+	FText GetGenerateButtonText() const
 	{
-		return (LODCount >= 2) ? LOCTEXT("GenerateLODChain", "Generate LOD Chain").ToString() : LOCTEXT("GenerateLOD", "Generate LOD").ToString();
+		return (LODCount >= 2) ? LOCTEXT("GenerateLODChain", "Generate LOD Chain") : LOCTEXT("GenerateLOD", "Generate LOD");
 	}
 
 	/** Callback function whenever the LOD count spinbox is changed. */
@@ -576,7 +576,7 @@ void SSkeletalSimplificationOptions::Construct(const FArguments& InArgs)
 			.Padding(0.0f, 8.0f, 0.0f, 0.0f)
 			[
 				SNew(STextBlock)
-				.Text(LOCTEXT("DesiredQuality_MeshSimplification", "Desired Quality").ToString())	
+				.Text(LOCTEXT("DesiredQuality_MeshSimplification", "Desired Quality"))	
 			]
 
 			+SVerticalBox::Slot()
@@ -591,7 +591,7 @@ void SSkeletalSimplificationOptions::Construct(const FArguments& InArgs)
 				.VAlign(VAlign_Center)
 				[
 					SNew(STextBlock)
-					.Text(LOCTEXT("QualityReductionMethod", "Simplification Type").ToString())
+					.Text(LOCTEXT("QualityReductionMethod", "Simplification Type"))
 				]
 
 				+SHorizontalBox::Slot()
@@ -654,7 +654,7 @@ void SSkeletalSimplificationOptions::Construct(const FArguments& InArgs)
 				.VAlign(VAlign_Center)
 				[
 					SNew(STextBlock)
-					.Text(LOCTEXT("Silhouette_MeshSimplification", "Silhouette").ToString())
+					.Text(LOCTEXT("Silhouette_MeshSimplification", "Silhouette"))
 				]
 
 				+SHorizontalBox::Slot()
@@ -678,7 +678,7 @@ void SSkeletalSimplificationOptions::Construct(const FArguments& InArgs)
 				.VAlign(VAlign_Center)
 				[
 					SNew(STextBlock)
-					.Text(LOCTEXT("Texture_MeshSimplification", "Texture").ToString())
+					.Text(LOCTEXT("Texture_MeshSimplification", "Texture"))
 				]
 
 				+SHorizontalBox::Slot()
@@ -702,7 +702,7 @@ void SSkeletalSimplificationOptions::Construct(const FArguments& InArgs)
 				.VAlign(VAlign_Center)
 				[
 					SNew(STextBlock)
-					.Text(LOCTEXT("Shading_MeshSimplification", "Shading").ToString())
+					.Text(LOCTEXT("Shading_MeshSimplification", "Shading"))
 				]
 
 				+SHorizontalBox::Slot()
@@ -726,7 +726,7 @@ void SSkeletalSimplificationOptions::Construct(const FArguments& InArgs)
 				.VAlign(VAlign_Center)
 				[
 					SNew(STextBlock)
-					.Text(LOCTEXT("Skinning_MeshSimplification", "Skinning").ToString())
+					.Text(LOCTEXT("Skinning_MeshSimplification", "Skinning"))
 				]
 
 				+SHorizontalBox::Slot()
@@ -744,7 +744,7 @@ void SSkeletalSimplificationOptions::Construct(const FArguments& InArgs)
 			.Padding(0.0f, 8.0f, 0.0f, 0.0f)
 			[
 				SNew(STextBlock)
-				.Text(LOCTEXT("RepairOptions", "Repair Options").ToString())	
+				.Text(LOCTEXT("RepairOptions", "Repair Options"))	
 			]
 
 			+SVerticalBox::Slot()
@@ -759,7 +759,7 @@ void SSkeletalSimplificationOptions::Construct(const FArguments& InArgs)
 				.VAlign(VAlign_Center)
 				[
 					SNew(STextBlock)
-					.Text(LOCTEXT("WeldingThreshold", "Welding Threshold").ToString())
+					.Text(LOCTEXT("WeldingThreshold", "Welding Threshold"))
 				]
 
 				+SHorizontalBox::Slot()
@@ -783,7 +783,7 @@ void SSkeletalSimplificationOptions::Construct(const FArguments& InArgs)
 				SAssignNew(RecomputeNormalsCheckbox, SCheckBox)
 				[
 					SNew(STextBlock)
-					.Text(LOCTEXT("RecomputeNormals", "Recompute Normals").ToString())
+					.Text(LOCTEXT("RecomputeNormals", "Recompute Normals"))
 				]
 			]
 
@@ -799,7 +799,7 @@ void SSkeletalSimplificationOptions::Construct(const FArguments& InArgs)
 				.VAlign(VAlign_Center)
 				[
 					SNew(STextBlock)
-					.Text(LOCTEXT("HardEdgeAngleThreshold", "Hard Edge Angle Threshold").ToString())
+					.Text(LOCTEXT("HardEdgeAngleThreshold", "Hard Edge Angle Threshold"))
 				]
 
 				+SHorizontalBox::Slot()
@@ -821,7 +821,7 @@ void SSkeletalSimplificationOptions::Construct(const FArguments& InArgs)
 			.Padding(0.0f, 8.0f, 0.0f, 0.0f)
 			[
 				SNew(STextBlock)
-				.Text(LOCTEXT("SkeletonSimplification", "Skeleton Simplification").ToString())	
+				.Text(LOCTEXT("SkeletonSimplification", "Skeleton Simplification"))	
 			]
 
 			+SVerticalBox::Slot()
@@ -835,8 +835,8 @@ void SSkeletalSimplificationOptions::Construct(const FArguments& InArgs)
 				.VAlign(VAlign_Center)
 				[
 					SNew(STextBlock)
-					.Text(LOCTEXT("SkeletonSimplification_BoneReduction", "Bones to Remove").ToString())
-					.ToolTipText(LOCTEXT("SkeletonSimplification_BoneReductionToolTip", "You can choose these bones in skeletaon tree as well.").ToString())
+					.Text(LOCTEXT("SkeletonSimplification_BoneReduction", "Bones to Remove"))
+					.ToolTipText(LOCTEXT("SkeletonSimplification_BoneReductionToolTip", "You can choose these bones in skeletaon tree as well."))
 				]
 
 				+SHorizontalBox::Slot()
@@ -903,7 +903,7 @@ void SSkeletalSimplificationOptions::Construct(const FArguments& InArgs)
 				.VAlign(VAlign_Center)
 				[
 					SNew(STextBlock)
-					.Text(LOCTEXT("MaxBonesPerVertex", "Max Bones Per Vertex").ToString())
+					.Text(LOCTEXT("MaxBonesPerVertex", "Max Bones Per Vertex"))
 				]
 
 				+SHorizontalBox::Slot()
diff --git a/Engine/Source/Editor/UnrealEd/Private/AnimationCompressionPanel.cpp b/Engine/Source/Editor/UnrealEd/Private/AnimationCompressionPanel.cpp
--- a/Engine/Source/Editor/UnrealEd/Private/AnimationCompressionPanel.cpp
+++ b/Engine/Source/Editor/UnrealEd/Private/AnimationCompressionPanel.cpp
@@ -98,7 +98,7 @@ void SAnimationCompressionPanel::Construct(const FArguments& InArgs)
 				+SUniformGridPanel::Slot(0,0)
 				[
 					SNew(SButton)
-					.Text(LOCTEXT("AnimCompressionApply", "Apply").ToString())
+					.Text(LOCTEXT("AnimCompressionApply", "Apply"))
 					.HAlign(HAlign_Center)
 					.ContentPadding(FEditorStyle::GetMargin("StandardDialog.ContentPadding"))
 					.OnClicked(this, &SAnimationCompressionPanel::ApplyClicked)
diff --git a/Engine/Source/Editor/UnrealEd/Private/ConsolidateWindow.cpp b/Engine/Source/Editor/UnrealEd/Private/ConsolidateWindow.cpp
--- a/Engine/Source/Editor/UnrealEd/Private/ConsolidateWindow.cpp
+++ b/Engine/Source/Editor/UnrealEd/Private/ConsolidateWindow.cpp
@@ -178,7 +178,7 @@ private:
 	 * @note	The current implementation is naive and will wipe out any pre-existing message when called. The current needs of the window don't require
 	 *			anything more sophisticated, but in the future perhaps the messages should be appended, multiple panel types should exist, etc.
 	 */
-	void DisplayMessage( bool bError, const FString& ErrorMessage );
+	void DisplayMessage( bool bError, const FText& ErrorMessage );
 	
 	// Button Responses
 
@@ -352,7 +352,7 @@ void SConsolidateToolWidget::Construct( const FArguments& InArgs )
 		[
 			SNew( STextBlock )
 			.AutoWrapText(true)
-			.Text( LOCTEXT("Consolidate_Select", "Select an asset to serve as the asset to consolidate the non-selected assets to. This will replace all uses of the non-selected assets below with the selected asset.").ToString() )
+			.Text( LOCTEXT("Consolidate_Select", "Select an asset to serve as the asset to consolidate the non-selected assets to. This will replace all uses of the non-selected assets below with the selected asset.") )
 		]
 
 		+SVerticalBox::Slot()
@@ -536,7 +536,7 @@ int32 SConsolidateToolWidget::AddConsolidationObjects( const TArray<UObject*>& I
 	// display a warning message informing them of the potential danger.
 	if ( !AreObjClassesHomogeneous() && !bAlreadyWarnedAboutTypes )
 	{
-		DisplayMessage( false, LOCTEXT("Consolidate_WarningSameClass", "The object to consolidate are not the same class").ToString());
+		DisplayMessage( false, LOCTEXT("Consolidate_WarningSameClass", "The object to consolidate are not the same class"));
 		bAlreadyWarnedAboutTypes = true;
 	}
 
@@ -735,7 +735,7 @@ void SConsolidateToolWidget::RemoveSelectedObject()
  * @note	The current implementation is naive and will wipe out any pre-existing message when called. The current needs of the window don't require
  *			anything more sophisticated, but in the future perhaps the messages should be appended, multiple panel types should exist, etc.
  */
-void SConsolidateToolWidget::DisplayMessage( bool bError, const FString& ErrorMessage )
+void SConsolidateToolWidget::DisplayMessage( bool bError, const FText& ErrorMessage )
 {
 	// Update the error text block to display the requested message
 	ErrorPanel->SetError(ErrorMessage);
@@ -794,7 +794,7 @@ FReply SConsolidateToolWidget::OnConsolidateButtonClicked()
 		// If the consolidation resulted in failed (partially consolidated) objects, do not save, and inform the user no save attempt was made
 		else if ( ConsResults.FailedConsolidationObjs.Num() > 0 && bSavePackagesChecked == true )
 		{
-			DisplayMessage( true, LOCTEXT("Consolidate_WarningPartial", "Not all objects could be consolidated, no save has occurred").ToString() );
+			DisplayMessage( true, LOCTEXT("Consolidate_WarningPartial", "Not all objects could be consolidated, no save has occurred") );
 		}
 	}
 
diff --git a/Engine/Source/Editor/UnrealEd/Private/Factories/AnimBlueprintFactory.cpp b/Engine/Source/Editor/UnrealEd/Private/Factories/AnimBlueprintFactory.cpp
--- a/Engine/Source/Editor/UnrealEd/Private/Factories/AnimBlueprintFactory.cpp
+++ b/Engine/Source/Editor/UnrealEd/Private/Factories/AnimBlueprintFactory.cpp
@@ -112,7 +112,7 @@ public:
 							.HAlign(HAlign_Center)
 							.ContentPadding( FEditorStyle::GetMargin("StandardDialog.ContentPadding") )
 							.OnClicked(this, &SAnimBlueprintCreateDialog::OkClicked)	
-							.Text(LOCTEXT("CreateAnimBlueprintOk", "OK").ToString())
+							.Text(LOCTEXT("CreateAnimBlueprintOk", "OK"))
 						]
 						+SUniformGridPanel::Slot(1,0)
 						[
@@ -120,7 +120,7 @@ public:
 							.HAlign(HAlign_Center)
 							.ContentPadding( FEditorStyle::GetMargin("StandardDialog.ContentPadding") )
 							.OnClicked(this, &SAnimBlueprintCreateDialog::CancelClicked)
-							.Text(LOCTEXT("CreateAnimBlueprintCancel", "Cancel").ToString())
+							.Text(LOCTEXT("CreateAnimBlueprintCancel", "Cancel"))
 						]
 					]
 				]
diff --git a/Engine/Source/Editor/UnrealEd/Private/Factories/CSVImportFactory.cpp b/Engine/Source/Editor/UnrealEd/Private/Factories/CSVImportFactory.cpp
--- a/Engine/Source/Editor/UnrealEd/Private/Factories/CSVImportFactory.cpp
+++ b/Engine/Source/Editor/UnrealEd/Private/Factories/CSVImportFactory.cpp
@@ -162,7 +162,7 @@ public:
 				.AutoHeight()
 				[
 					SNew(STextBlock)
-					.Text( LOCTEXT("ChooseCurveType", "Choose Curve Interpolation Type:").ToString() )
+					.Text( LOCTEXT("ChooseCurveType", "Choose Curve Interpolation Type:") )
 					.Visibility( this, &SCSVImportOptions::GetCurveTypeVis )
 				]
 				+SVerticalBox::Slot()
@@ -186,14 +186,14 @@ public:
 					.AutoWidth()
 					[
 						SNew(SButton)
-						.Text(LOCTEXT("OK", "OK").ToString())
+						.Text(LOCTEXT("OK", "OK"))
 						.OnClicked( this, &SCSVImportOptions::OnImport )
 					]
 					+SHorizontalBox::Slot()
 					.AutoWidth()
 					[
 						SNew(SButton)
-						.Text(LOCTEXT("Cancel", "Cancel").ToString())
+						.Text(LOCTEXT("Cancel", "Cancel"))
 						.OnClicked( this, &SCSVImportOptions::OnCancel )
 					]
 				]
