#UE4 - Fix Build Scale property look

diff --git a/Engine/Source/Editor/StaticMeshEditor/Private/StaticMeshEditorTools.cpp b/Engine/Source/Editor/StaticMeshEditor/Private/StaticMeshEditorTools.cpp
--- a/Engine/Source/Editor/StaticMeshEditor/Private/StaticMeshEditorTools.cpp
+++ b/Engine/Source/Editor/StaticMeshEditor/Private/StaticMeshEditorTools.cpp
@@ -10,6 +10,7 @@
 #include "PropertyCustomizationHelpers.h"
 #include "PhysicsEngine/BodySetup.h"
 #include "FbxMeshUtils.h"
+#include "SVectorInputBox.h"
 
 const int32 DefaultHullCount = 4;
 const int32 DefaultVertsPerHull = 12;
@@ -343,59 +344,18 @@ void FMeshBuildSettingsLayout::GenerateChildContent( IDetailChildrenBuilder& Chi
 			.ToolTipText( LOCTEXT("BuildScale_ToolTip", "The local scale applied when building the mesh") )
 		]
 		.ValueContent()
+		.MinDesiredWidth(125.0f * 3.0f)
+		.MaxDesiredWidth(125.0f * 3.0f)
 		[
-			SNew( SVerticalBox )
-			+SVerticalBox::Slot()
-			.Padding( FMargin(0.0f, 1.0f, 0.0f, 6.0f ) )
-			[
-				SNew(SNumericEntryBox<float>)
-				.AllowSpin(true)
-				.Font(IDetailLayoutBuilder::GetDetailFont())
-				.MinValue(0.0f)
-				.MaxValue(TOptional<float>())
-				.MaxSliderValue(TOptional<float>())
-				.MinSliderValue(TOptional<float>())
-				.Value( this, &FMeshBuildSettingsLayout::GetBuildScaleX)
-				.OnValueChanged(this, &FMeshBuildSettingsLayout::OnBuildScaleXChanged)
-				.Label()
-				[
-					SNumericEntryBox<float>::BuildLabel( LOCTEXT("BuildScaleX_Label", "X"), FLinearColor::White, SNumericEntryBox<float>::RedLabelBackgroundColor )
-				]
-			]
-			+SVerticalBox::Slot()
-			.Padding( FMargin(0.0f, 1.0f, 0.0f, 6.0f ) )
-			[
-				SNew(SNumericEntryBox<float>)
-				.AllowSpin(true)
-				.Font(IDetailLayoutBuilder::GetDetailFont())
-				.MinValue(0.0f)
-				.MaxValue(TOptional<float>())
-				.MaxSliderValue(TOptional<float>())
-				.MinSliderValue(TOptional<float>())
-				.Value( this, &FMeshBuildSettingsLayout::GetBuildScaleY)
-				.OnValueChanged(this, &FMeshBuildSettingsLayout::OnBuildScaleYChanged)
-				.Label()
-				[
-					SNumericEntryBox<float>::BuildLabel( LOCTEXT("BuildScaleY_Label", "Y"), FLinearColor::White, SNumericEntryBox<float>::GreenLabelBackgroundColor )
-				]
-			]
-			+SVerticalBox::Slot()
-			.Padding( FMargin(0.0f, 1.0f, 0.0f, 2.0f ) )
-			[
-				SNew(SNumericEntryBox<float>)
-				.AllowSpin(true)
-				.Font(IDetailLayoutBuilder::GetDetailFont())
-				.MinValue(0.0f)
-				.MaxValue(TOptional<float>())
-				.MaxSliderValue(TOptional<float>())
-				.MinSliderValue(TOptional<float>())
-				.Value( this, &FMeshBuildSettingsLayout::GetBuildScaleZ)
-				.OnValueChanged(this, &FMeshBuildSettingsLayout::OnBuildScaleZChanged)
-				.Label()
-				[
-					SNumericEntryBox<float>::BuildLabel( LOCTEXT("BuildScaleZ_Label", "Z"), FLinearColor::White, SNumericEntryBox<float>::BlueLabelBackgroundColor )
-				]
-			]
+			SNew(SVectorInputBox)
+			.X(this, &FMeshBuildSettingsLayout::GetBuildScaleX)
+			.Y(this, &FMeshBuildSettingsLayout::GetBuildScaleY)
+			.Z(this, &FMeshBuildSettingsLayout::GetBuildScaleZ)
+			.bColorAxisLabels(false)
+			.OnXCommitted(this, &FMeshBuildSettingsLayout::OnBuildScaleXChanged)
+			.OnYCommitted(this, &FMeshBuildSettingsLayout::OnBuildScaleYChanged)
+			.OnZCommitted(this, &FMeshBuildSettingsLayout::OnBuildScaleXChanged)
+			.Font(IDetailLayoutBuilder::GetDetailFont())
 		];
 	}
 		
@@ -485,17 +445,17 @@ void FMeshBuildSettingsLayout::OnUseFullPrecisionUVsChanged(ESlateCheckBoxState:
 	BuildSettings.bUseFullPrecisionUVs = (NewState == ESlateCheckBoxState::Checked) ? true : false;
 }
 
-void FMeshBuildSettingsLayout::OnBuildScaleXChanged( float NewScaleX )
+void FMeshBuildSettingsLayout::OnBuildScaleXChanged( float NewScaleX, ETextCommit::Type TextCommitType )
 {
 	BuildSettings.BuildScale3D.X = NewScaleX;
 }
 
-void FMeshBuildSettingsLayout::OnBuildScaleYChanged( float NewScaleY )
+void FMeshBuildSettingsLayout::OnBuildScaleYChanged( float NewScaleY, ETextCommit::Type TextCommitType )
 {
 	BuildSettings.BuildScale3D.Y = NewScaleY;
 }
 
-void FMeshBuildSettingsLayout::OnBuildScaleZChanged( float NewScaleZ )
+void FMeshBuildSettingsLayout::OnBuildScaleZChanged( float NewScaleZ, ETextCommit::Type TextCommitType )
 {
 	BuildSettings.BuildScale3D.Z = NewScaleZ;
 }
diff --git a/Engine/Source/Editor/StaticMeshEditor/Private/StaticMeshEditorTools.h b/Engine/Source/Editor/StaticMeshEditor/Private/StaticMeshEditorTools.h
--- a/Engine/Source/Editor/StaticMeshEditor/Private/StaticMeshEditorTools.h
+++ b/Engine/Source/Editor/StaticMeshEditor/Private/StaticMeshEditorTools.h
@@ -145,9 +145,9 @@ private:
 	void OnRecomputeTangentsChanged(ESlateCheckBoxState::Type NewState);
 	void OnRemoveDegeneratesChanged(ESlateCheckBoxState::Type NewState);
 	void OnUseFullPrecisionUVsChanged(ESlateCheckBoxState::Type NewState);
-	void OnBuildScaleXChanged( float NewScaleX );
-	void OnBuildScaleYChanged( float NewScaleY );
-	void OnBuildScaleZChanged( float NewScaleZ );
+	void OnBuildScaleXChanged( float NewScaleX, ETextCommit::Type TextCommitType );
+	void OnBuildScaleYChanged( float NewScaleY, ETextCommit::Type TextCommitType );
+	void OnBuildScaleZChanged( float NewScaleZ, ETextCommit::Type TextCommitType );
 
 private:
 	TWeakPtr<FLevelOfDetailSettingsLayout> ParentLODSettings;
