#UE4 - Fix typo in delegate assignment causing buildscale.z in the static mesh editor not working

diff --git a/Engine/Source/Editor/StaticMeshEditor/Private/StaticMeshEditorTools.cpp b/Engine/Source/Editor/StaticMeshEditor/Private/StaticMeshEditorTools.cpp
--- a/Engine/Source/Editor/StaticMeshEditor/Private/StaticMeshEditorTools.cpp
+++ b/Engine/Source/Editor/StaticMeshEditor/Private/StaticMeshEditorTools.cpp
@@ -354,7 +354,7 @@ void FMeshBuildSettingsLayout::GenerateChildContent( IDetailChildrenBuilder& Chi
 			.bColorAxisLabels(false)
 			.OnXCommitted(this, &FMeshBuildSettingsLayout::OnBuildScaleXChanged)
 			.OnYCommitted(this, &FMeshBuildSettingsLayout::OnBuildScaleYChanged)
-			.OnZCommitted(this, &FMeshBuildSettingsLayout::OnBuildScaleXChanged)
+			.OnZCommitted(this, &FMeshBuildSettingsLayout::OnBuildScaleZChanged)
 			.Font(IDetailLayoutBuilder::GetDetailFont())
 		];
 	}
