Disallowed Build Scales of 0 in the LODn details.
#ttp 333283 - EDITOR: Static Mesh will disappear when changing the Z axis from 0 to any other number while using the Build Scale tool
#branch UE4
#proj Editor.StaticMeshEditor
#change Trapped inputs close to zero in FMeshBuildSettingsLayout::OnBuildScale*Changed and rejected them.

diff --git a/Engine/Source/Editor/StaticMeshEditor/Private/StaticMeshEditorTools.cpp b/Engine/Source/Editor/StaticMeshEditor/Private/StaticMeshEditorTools.cpp
--- a/Engine/Source/Editor/StaticMeshEditor/Private/StaticMeshEditorTools.cpp
+++ b/Engine/Source/Editor/StaticMeshEditor/Private/StaticMeshEditorTools.cpp
@@ -447,17 +447,26 @@ void FMeshBuildSettingsLayout::OnUseFullPrecisionUVsChanged(ESlateCheckBoxState:
 
 void FMeshBuildSettingsLayout::OnBuildScaleXChanged( float NewScaleX, ETextCommit::Type TextCommitType )
 {
-	BuildSettings.BuildScale3D.X = NewScaleX;
+	if (!FMath::IsNearlyEqual(NewScaleX, 0.0f))
+	{
+		BuildSettings.BuildScale3D.X = NewScaleX;
+	}
 }
 
 void FMeshBuildSettingsLayout::OnBuildScaleYChanged( float NewScaleY, ETextCommit::Type TextCommitType )
 {
-	BuildSettings.BuildScale3D.Y = NewScaleY;
+	if (!FMath::IsNearlyEqual(NewScaleY, 0.0f))
+	{
+		BuildSettings.BuildScale3D.Y = NewScaleY;
+	}
 }
 
 void FMeshBuildSettingsLayout::OnBuildScaleZChanged( float NewScaleZ, ETextCommit::Type TextCommitType )
 {
-	BuildSettings.BuildScale3D.Z = NewScaleZ;
+	if (!FMath::IsNearlyEqual(NewScaleZ, 0.0f))
+	{
+		BuildSettings.BuildScale3D.Z = NewScaleZ;
+	}
 }
 
 FMeshReductionSettingsLayout::FMeshReductionSettingsLayout( TSharedRef<FLevelOfDetailSettingsLayout> InParentLODSettings )
