Made Static Meshes use ScreenSize instead of Distances for LOD level calculation
#ttp 319896 - Skeletalmesh/Staticmesh LOD refactoring: Part 1 (Refactor Static Meshes to use ScreenSize rather than distance)
#branch UE4
#change Removed LOD Distances and added ScreenSize for Static Meshes

diff --git a/Engine/Source/Editor/StaticMeshEditor/Private/StaticMeshEditorTools.cpp b/Engine/Source/Editor/StaticMeshEditor/Private/StaticMeshEditorTools.cpp
--- a/Engine/Source/Editor/StaticMeshEditor/Private/StaticMeshEditorTools.cpp
+++ b/Engine/Source/Editor/StaticMeshEditor/Private/StaticMeshEditorTools.cpp
@@ -980,10 +980,11 @@ FLevelOfDetailSettingsLayout::FLevelOfDetailSettingsLayout( FStaticMeshEditor& I
 
 	for (int32 i = 0; i < MAX_STATIC_MESH_LODS; ++i)
 	{
-		LODDistances[i] = 0.0f;
 		bBuildSettingsExpanded[i] = false;
 		bReductionSettingsExpanded[i] = false;
 		bSectionSettingsExpanded[i] = (i == 0);
+
+		LODScreenSizes[i] = 0.0f;
 	}
 
 	LODCount = StaticMeshEditor.GetStaticMesh()->GetNumLODs();
@@ -1148,7 +1149,7 @@ void FLevelOfDetailSettingsLayout::AddLODLevelCategories( IDetailLayoutBuilder&
 					BuildSettingsWidgets[LODIndex]->UpdateSettings(SrcModel.BuildSettings);
 				}
 
-				LODDistances[LODIndex] = SrcModel.LODDistance;
+				LODScreenSizes[LODIndex] = SrcModel.ScreenSize;
 			}
 			else if (LODIndex > 0)
 			{
@@ -1160,10 +1161,10 @@ void FLevelOfDetailSettingsLayout::AddLODLevelCategories( IDetailLayoutBuilder&
 					ReductionSettingsWidgets[LODIndex]->UpdateSettings(ReductionSettings);
 				}
 
-				if (LODDistances[LODIndex] <= LODDistances[LODIndex-1])
+				if(LODScreenSizes[LODIndex] >= LODScreenSizes[LODIndex-1])
 				{
-					const float DefaultLODDistance = 1000.0f;
-					LODDistances[LODIndex] = LODDistances[LODIndex-1] + DefaultLODDistance;
+					const float DefaultScreenSizeDifference = 0.01f;
+					LODScreenSizes[LODIndex] = LODScreenSizes[LODIndex-1] - DefaultScreenSizeDifference;
 				}
 			}
 
@@ -1181,12 +1182,12 @@ void FLevelOfDetailSettingsLayout::AddLODLevelCategories( IDetailLayoutBuilder&
 				[
 					SNew( SHorizontalBox )
 					+ SHorizontalBox::Slot()
-					.Padding( FMargin( 5.0f, 0.0f ) )
+					.Padding(FMargin(5.0f, 0.0f))
 					.AutoWidth()
 					[
 						SNew(STextBlock)
 						.Font(FEditorStyle::GetFontStyle("StaticMeshEditor.NormalFont"))
-						.Text(this, &FLevelOfDetailSettingsLayout::GetLODDistanceTitle, LODIndex)
+						.Text(this, &FLevelOfDetailSettingsLayout::GetLODScreenSizeTitle, LODIndex)
 						.Visibility( LODIndex > 0 ? EVisibility::Visible : EVisibility::Collapsed )
 					]
 					+ SHorizontalBox::Slot()
@@ -1212,28 +1213,25 @@ void FLevelOfDetailSettingsLayout::AddLODLevelCategories( IDetailLayoutBuilder&
 			SectionSettingsWidgets[ LODIndex ] = MakeShareable( new FMeshSectionSettingsLayout( StaticMeshEditor, LODIndex ) );
 			SectionSettingsWidgets[ LODIndex ]->AddToCategory( LODCategory );
 
-			if (LODIndex > 0)
-			{
-				LODCategory.AddCustomRow( LOCTEXT("Distance", "Distance").ToString() )
-				.NameContent()
-				[
-					SNew(STextBlock)
-					.Font( IDetailLayoutBuilder::GetDetailFont() )
-					.Text(LOCTEXT("Distance", "Distance"))
-				]
-				.ValueContent()
-				[
-					SNew(SSpinBox<float>)
-					.Font( IDetailLayoutBuilder::GetDetailFont() )
-					.MinValue(0.0f)
-					.MaxValue(WORLD_MAX)
-					.SliderExponent(2.0f)
-					.Value(this, &FLevelOfDetailSettingsLayout::GetLODDistance, LODIndex)
-					.OnValueChanged(this, &FLevelOfDetailSettingsLayout::OnLODDistanceChanged, LODIndex)
-					.OnValueCommitted(this, &FLevelOfDetailSettingsLayout::OnLODDistanceCommitted, LODIndex)
-					.IsEnabled(this, &FLevelOfDetailSettingsLayout::CanChangeLODDistance)
-				];
-			}
+			LODCategory.AddCustomRow(( LOCTEXT("ScreenSizeRow", "ScreenSize").ToString()))
+			.NameContent()
+			[
+				SNew(STextBlock)
+				.Font(IDetailLayoutBuilder::GetDetailFont())
+				.Text(LOCTEXT("ScreenSizeName", "Screen Size"))
+			]
+			.ValueContent()
+			[
+				SNew(SSpinBox<float>)
+				.Font(IDetailLayoutBuilder::GetDetailFont())
+				.MinValue(0.0f)
+				.MaxValue(WORLD_MAX)
+				.SliderExponent(2.0f)
+				.Value(this, &FLevelOfDetailSettingsLayout::GetLODScreenSize, LODIndex)
+				.OnValueChanged(this, &FLevelOfDetailSettingsLayout::OnLODScreenSizeChanged, LODIndex)
+				.OnValueCommitted(this, &FLevelOfDetailSettingsLayout::OnLODScreenSizeCommitted, LODIndex)
+				.IsEnabled(this, &FLevelOfDetailSettingsLayout::CanChangeLODScreenSize)
+			];
 
 			if (BuildSettingsWidgets[LODIndex].IsValid())
 			{
@@ -1260,62 +1258,58 @@ int32 FLevelOfDetailSettingsLayout::GetLODCount() const
 	return LODCount;
 }
 
-float FLevelOfDetailSettingsLayout::GetLODDistance(int32 LODIndex) const
+float FLevelOfDetailSettingsLayout::GetLODScreenSize( int32 LODIndex ) const
 {
-	check(LODIndex > 0 && LODIndex < MAX_STATIC_MESH_LODS);
-	UStaticMesh* StaticMesh = StaticMeshEditor.GetStaticMesh();
-	float Distance = LODDistances[FMath::Clamp<int32>(LODIndex,0,MAX_STATIC_MESH_LODS-1)];
-	if (StaticMesh->bAutoComputeLODDistance)
+	check(LODIndex < MAX_STATIC_MESH_LODS);
+	UStaticMesh* Mesh = StaticMeshEditor.GetStaticMesh();
+	float ScreenSize = LODScreenSizes[FMath::Clamp(LODIndex, 0, MAX_STATIC_MESH_LODS-1)];
+	if(Mesh->bAutoComputeLODScreenSize)
 	{
-		Distance = StaticMesh->RenderData->LODDistance[LODIndex];
-		if (Distance == FLT_MAX)
-		{
-			Distance = 0.0f;
-		}
+		ScreenSize = Mesh->RenderData->ScreenSize[LODIndex];
 	}
-	else if (StaticMesh->SourceModels.IsValidIndex(LODIndex))
+	else if(Mesh->SourceModels.IsValidIndex(LODIndex))
 	{
-		Distance = StaticMesh->SourceModels[LODIndex].LODDistance;
+		ScreenSize = Mesh->SourceModels[LODIndex].ScreenSize;
 	}
-	return Distance;
+	return ScreenSize;
 }
 
-FText FLevelOfDetailSettingsLayout::GetLODDistanceTitle(int32 LODIndex) const
+FText FLevelOfDetailSettingsLayout::GetLODScreenSizeTitle( int32 LODIndex ) const
 {
-	return FText::Format( LOCTEXT("Distance_MeshSimplification", "Distance: {0}"), FText::AsNumber( LODIndex > 0 ? GetLODDistance( LODIndex ) : 0.0f ) );
+	return FText::Format( LOCTEXT("ScreenSize_MeshSimplification", "Screen Size: {0}"), FText::AsNumber(GetLODScreenSize(LODIndex)));
 }
 
-bool FLevelOfDetailSettingsLayout::CanChangeLODDistance() const
+bool FLevelOfDetailSettingsLayout::CanChangeLODScreenSize() const
 {
 	return !IsAutoLODEnabled();
 }
 
-void FLevelOfDetailSettingsLayout::OnLODDistanceChanged(float NewValue, int32 LODIndex)
+void FLevelOfDetailSettingsLayout::OnLODScreenSizeChanged( float NewValue, int32 LODIndex )
 {
-	check(LODIndex > 0 && LODIndex < MAX_STATIC_MESH_LODS);
+	check(LODIndex < MAX_STATIC_MESH_LODS);
 	UStaticMesh* StaticMesh = StaticMeshEditor.GetStaticMesh();
-	if (!StaticMesh->bAutoComputeLODDistance)
+	if (!StaticMesh->bAutoComputeLODScreenSize)
 	{
 		// First propagate any changes from the source models to our local scratch.
 		for (int32 i = 0; i < StaticMesh->SourceModels.Num(); ++i)
 		{
-			LODDistances[i] = StaticMesh->SourceModels[i].LODDistance;
+			LODScreenSizes[i] = StaticMesh->SourceModels[i].ScreenSize;
 		}
 
-		// Update LOD distances for all further LODs.
-		float Delta = NewValue - LODDistances[LODIndex];
-		LODDistances[LODIndex] = NewValue;
+		// Update Display factors for further LODs
+		float Delta = NewValue - LODScreenSizes[LODIndex];
+		LODScreenSizes[LODIndex] = NewValue;
 		for (int32 i = LODIndex + 1; i < MAX_STATIC_MESH_LODS; ++i)
 		{
-			LODDistances[i] += Delta;
+			LODScreenSizes[i] += Delta;
 		}
 
-		// Enforce that LODs don't swap in before the previous LOD.
-		const float MinLODDistance = 100.0f;
+		// Make sure we aren't trying to ovelap or have more than one LOD for a value
+		const float MinimumDifferenceInScreenSize = 0.01f;
 		for (int32 i = 1; i < MAX_STATIC_MESH_LODS; ++i)
 		{
-			float MinDistance = LODDistances[i-1] + MinLODDistance;
-			LODDistances[i] = FMath::Max(LODDistances[i], MinDistance);
+			float MaxValue = FMath::Clamp(LODScreenSizes[i-1] - MinimumDifferenceInScreenSize, 0.0f, 1.0f);
+			LODScreenSizes[i] = FMath::Min(LODScreenSizes[i], MaxValue);
 		}
 
 		// Push changes immediately.
@@ -1323,12 +1317,12 @@ void FLevelOfDetailSettingsLayout::OnLODDistanceChanged(float NewValue, int32 LO
 		{
 			if (StaticMesh->SourceModels.IsValidIndex(i))
 			{
-				StaticMesh->SourceModels[i].LODDistance = LODDistances[i];
+				StaticMesh->SourceModels[i].ScreenSize = LODScreenSizes[i];
 			}
 			if (StaticMesh->RenderData
 				&& StaticMesh->RenderData->LODResources.IsValidIndex(i))
 			{
-				StaticMesh->RenderData->LODDistance[i] = LODDistances[i];
+				StaticMesh->RenderData->ScreenSize[i] = LODScreenSizes[i];
 			}
 		}
 
@@ -1342,9 +1336,9 @@ void FLevelOfDetailSettingsLayout::OnLODDistanceChanged(float NewValue, int32 LO
 	}
 }
 
-void FLevelOfDetailSettingsLayout::OnLODDistanceCommitted(float NewValue, ETextCommit::Type CommitType, int32 LODIndex)
+void FLevelOfDetailSettingsLayout::OnLODScreenSizeCommitted( float NewValue, ETextCommit::Type CommitType, int32 LODIndex )
 {
-	OnLODDistanceChanged(NewValue, LODIndex);
+	OnLODScreenSizeChanged(NewValue, LODIndex);
 }
 
 void FLevelOfDetailSettingsLayout::UpdateLODNames()
@@ -1414,7 +1408,7 @@ void FLevelOfDetailSettingsLayout::OnLODGroupChanged(TSharedPtr<FString> NewValu
 			{
 				StaticMesh->SourceModels[LODIndex].ReductionSettings = GroupSettings.GetDefaultSettings(LODIndex);
 			}
-			StaticMesh->bAutoComputeLODDistance = true;
+			StaticMesh->bAutoComputeLODScreenSize = true;
 			StaticMesh->LightMapResolution = GroupSettings.GetDefaultLightMapResolution();
 		}
 		StaticMesh->PostEditChange();
@@ -1426,7 +1420,7 @@ bool FLevelOfDetailSettingsLayout::IsAutoLODEnabled() const
 {
 	UStaticMesh* StaticMesh = StaticMeshEditor.GetStaticMesh();
 	check(StaticMesh);
-	return StaticMesh->bAutoComputeLODDistance;
+	return StaticMesh->bAutoComputeLODScreenSize;
 }
 
 ESlateCheckBoxState::Type FLevelOfDetailSettingsLayout::IsAutoLODChecked() const
@@ -1439,16 +1433,16 @@ void FLevelOfDetailSettingsLayout::OnAutoLODChanged(ESlateCheckBoxState::Type Ne
 	UStaticMesh* StaticMesh = StaticMeshEditor.GetStaticMesh();
 	check(StaticMesh);
 	StaticMesh->Modify();
-	StaticMesh->bAutoComputeLODDistance = (NewState == ESlateCheckBoxState::Checked) ? true : false;
-	if (!StaticMesh->bAutoComputeLODDistance)
+	StaticMesh->bAutoComputeLODScreenSize = (NewState == ESlateCheckBoxState::Checked) ? true : false;
+	if (!StaticMesh->bAutoComputeLODScreenSize)
 	{
 		if (StaticMesh->SourceModels.IsValidIndex(0))
 		{
-			StaticMesh->SourceModels[0].LODDistance = 0.0f;
+			StaticMesh->SourceModels[0].ScreenSize = 1.0f;
 		}
 		for (int32 LODIndex = 1; LODIndex < StaticMesh->SourceModels.Num(); ++LODIndex)
 		{
-			StaticMesh->SourceModels[LODIndex].LODDistance = StaticMesh->RenderData->LODDistance[LODIndex];
+			StaticMesh->SourceModels[LODIndex].ScreenSize = StaticMesh->RenderData->ScreenSize[LODIndex];
 		}
 	}
 	StaticMesh->PostEditChange();
@@ -1556,16 +1550,19 @@ void FLevelOfDetailSettingsLayout::ApplyChanges()
 
 		if (LODIndex == 0)
 		{
-			SrcModel.LODDistance = 0.0f;
+			SrcModel.ScreenSize = 1.0f;
 		}
 		else
 		{
-			SrcModel.LODDistance = LODDistances[LODIndex];
+			SrcModel.ScreenSize = LODScreenSizes[LODIndex];
 			FStaticMeshSourceModel& PrevModel = StaticMesh->SourceModels[LODIndex-1];
-			if (SrcModel.LODDistance <= PrevModel.LODDistance)
+			if(SrcModel.ScreenSize >= PrevModel.ScreenSize)
 			{
-				const float DefaultLODDistance = 1000.0f;
-				SrcModel.LODDistance = PrevModel.LODDistance + DefaultLODDistance;
+				const float DefaultScreenSizeDifference = 0.01f;
+				LODScreenSizes[LODIndex] = LODScreenSizes[LODIndex-1] - DefaultScreenSizeDifference;
+
+				// Make sure there are no incorrectly overlapping values
+				SrcModel.ScreenSize = 1.0f - 0.01f * LODIndex;
 			}
 		}
 	}
diff --git a/Engine/Source/Editor/StaticMeshEditor/Private/StaticMeshEditorTools.h b/Engine/Source/Editor/StaticMeshEditor/Private/StaticMeshEditorTools.h
--- a/Engine/Source/Editor/StaticMeshEditor/Private/StaticMeshEditorTools.h
+++ b/Engine/Source/Editor/StaticMeshEditor/Private/StaticMeshEditorTools.h
@@ -258,11 +258,13 @@ private:
 	void OnLODCountChanged(int32 NewValue);
 	void OnLODCountCommitted(int32 InValue, ETextCommit::Type CommitInfo);
 	int32 GetLODCount() const;
-	float GetLODDistance(int32 LODIndex) const;
-	FText GetLODDistanceTitle(int32 LODIndex) const;
-	bool CanChangeLODDistance() const;
-	void OnLODDistanceChanged(float NewValue, int32 LODIndex);
-	void OnLODDistanceCommitted(float NewValue, ETextCommit::Type CommitType, int32 LODIndex);
+
+	float GetLODScreenSize(int32 LODIndex)const;
+	FText GetLODScreenSizeTitle(int32 LODIndex) const;
+	bool CanChangeLODScreenSize() const;
+	void OnLODScreenSizeChanged(float NewValue, int32 LODIndex);
+	void OnLODScreenSizeCommitted(float NewValue, ETextCommit::Type CommitType, int32 LODIndex);
+
 	void OnBuildSettingsExpanded(bool bIsExpanded, int32 LODIndex);
 	void OnReductionSettingsExpanded(bool bIsExpanded, int32 LODIndex);
 	void OnSectionSettingsExpanded(bool bIsExpanded, int32 LODIndex);
@@ -296,8 +298,8 @@ private:
 	TSharedPtr<FMeshBuildSettingsLayout> BuildSettingsWidgets[MAX_STATIC_MESH_LODS];
 	TSharedPtr<FMeshSectionSettingsLayout> SectionSettingsWidgets[MAX_STATIC_MESH_LODS];
 
-	/** The distances at which each LOD swaps in. */
-	float LODDistances[MAX_STATIC_MESH_LODS];
+	/** The display factors at which LODs swap */
+	float LODScreenSizes[MAX_STATIC_MESH_LODS];
 
 	/** Helper value that corresponds to the 'Number of LODs' spinbox.*/
 	int32 LODCount;
diff --git a/Engine/Source/Editor/StaticMeshEditor/Private/StaticMeshEditorViewportClient.cpp b/Engine/Source/Editor/StaticMeshEditor/Private/StaticMeshEditorViewportClient.cpp
--- a/Engine/Source/Editor/StaticMeshEditor/Private/StaticMeshEditorViewportClient.cpp
+++ b/Engine/Source/Editor/StaticMeshEditor/Private/StaticMeshEditorViewportClient.cpp
@@ -458,13 +458,7 @@ void FStaticMeshEditorViewportClient::DrawCanvas( FViewport& InViewport, FSceneV
 	int32 CurrentLODLevel = StaticMeshEditorPtr.Pin()->GetCurrentLODLevel();
 	if (CurrentLODLevel == 0)
 	{
-		// NOTE: Does not account for MaxDrawDistanceScale.
-		float DistanceSquared = (StaticMeshComponent->Bounds.Origin - View.ViewMatrices.ViewOrigin).SizeSquared();
-		const float LODFactorDistanceSquared = DistanceSquared * FMath::Square(View.LODDistanceFactor);
-		while (CurrentLODLevel < MAX_STATIC_MESH_LODS && LODFactorDistanceSquared >= FMath::Square(StaticMesh->RenderData->LODDistance[CurrentLODLevel+1]))
-		{
-			CurrentLODLevel++;
-		}
+		CurrentLODLevel = ComputeStaticMeshLOD(StaticMesh->RenderData, StaticMeshComponent->Bounds.Origin, StaticMeshComponent->Bounds.SphereRadius, View);
 	}
 	else
 	{
@@ -480,6 +474,20 @@ void FStaticMeshEditorViewportClient::DrawCanvas( FViewport& InViewport, FSceneV
 		);
 	YPos += 18;
 
+	float CurrentScreenSize = ComputeBoundsScreenSize(StaticMeshComponent->Bounds.Origin, StaticMeshComponent->Bounds.SphereRadius, View);
+	FNumberFormattingOptions FormatOptions;
+	FormatOptions.MinimumFractionalDigits = 3;
+	FormatOptions.MaximumFractionalDigits = 6;
+	FormatOptions.MaximumIntegralDigits = 6;
+	Canvas.DrawShadowedString(
+		6,
+		YPos,
+		*FText::Format( NSLOCTEXT("UnrealEd", "ScreenSize_F", "Current Screen Size:  {0}"), FText::AsNumber(CurrentScreenSize, &FormatOptions)).ToString(),
+		GEngine->GetSmallFont(),
+		FLinearColor::White
+		);
+	YPos += 18;
+
 	Canvas.DrawShadowedString(
 		6,
 		YPos,
diff --git a/Engine/Source/Runtime/Core/Public/UObject/ObjectVersion.h b/Engine/Source/Runtime/Core/Public/UObject/ObjectVersion.h
--- a/Engine/Source/Runtime/Core/Public/UObject/ObjectVersion.h
+++ b/Engine/Source/Runtime/Core/Public/UObject/ObjectVersion.h
@@ -536,6 +536,8 @@ enum EUnrealEngineObjectUE4Version
 	VER_UE4_VEHICLES_UNIT_CHANGE2,
 	// Changed order of interface class serialization
 	VER_UE4_UCLASS_SERIALIZE_INTERFACES_AFTER_LINKING,
+	// Change from LOD distances to display factors
+	VER_UE4_STATIC_MESH_SCREEN_SIZE_LODS,
  
 	// -----<new versions can be added before this line>-------------------------------------------------
 	// - this needs to be the last line (see note below)
diff --git a/Engine/Source/Runtime/Engine/Classes/Engine/StaticMesh.h b/Engine/Source/Runtime/Engine/Classes/Engine/StaticMesh.h
--- a/Engine/Source/Runtime/Engine/Classes/Engine/StaticMesh.h
+++ b/Engine/Source/Runtime/Engine/Classes/Engine/StaticMesh.h
@@ -178,11 +178,14 @@ struct FStaticMeshSourceModel
 
 	/** Reduction settings to apply when building render data. */
 	UPROPERTY(EditAnywhere, Category=ReductionSettings)
-	FMeshReductionSettings ReductionSettings;
+	FMeshReductionSettings ReductionSettings; 
 
-	/** The distance at which this LOD swaps in. */
+	UPROPERTY()
+	float LODDistance_DEPRECATED;
+
+	/** ScreenSize to display this LOD */
 	UPROPERTY(EditAnywhere, Category=ReductionSettings)
-	float LODDistance;
+	float ScreenSize;
 
 	/** Default constructor. */
 	ENGINE_API FStaticMeshSourceModel();
@@ -329,7 +332,13 @@ class UStaticMesh : public UObject, public IInterface_CollisionDataProvider, pub
 
 	/** If true, the distances at which LODs swap are computed automatically. */
 	UPROPERTY()
-	uint32 bAutoComputeLODDistance:1;
+	uint32 bAutoComputeLODScreenSize:1;
+
+	/**
+	* If true on post load we need to calculate Display Factors from the
+	* loaded LOD distances.
+	*/
+	bool bRequiresLODDistanceConversion : 1;
 #endif // #if WITH_EDITORONLY_DATA
 
 	/** Materials used by this static mesh. Individual sections index in to this array. */
@@ -410,6 +419,7 @@ class UStaticMesh : public UObject, public IInterface_CollisionDataProvider, pub
 	/** The stored camera position to use as a default for the static mesh editor */
 	UPROPERTY()
 	FAssetEditorOrbitCameraPosition EditorCameraPosition;
+
 #endif // WITH_EDITORONLY_DATA
 
 	/** For simplified meshes, this is the CRC of the high res mesh we were originally duplicated from. */
@@ -626,6 +636,11 @@ private:
 	 */
 	void SerializeLegacySouceData(class FArchive& Ar, const struct FBoxSphereBounds& LegacyBounds);
 
+	/**
+	 * Converts legacy LODDistance in the source models to Display Factor
+	 */
+	void ConvertLegacyLODDistance();
+
 	/**
 	 * Fixes up static meshes that were imported with sections that had zero triangles.
 	 */
diff --git a/Engine/Source/Runtime/Engine/Private/Landscape/LandscapeRender.cpp b/Engine/Source/Runtime/Engine/Private/Landscape/LandscapeRender.cpp
--- a/Engine/Source/Runtime/Engine/Private/Landscape/LandscapeRender.cpp
+++ b/Engine/Source/Runtime/Engine/Private/Landscape/LandscapeRender.cpp
@@ -1104,7 +1104,7 @@ void FLandscapeComponentSceneProxy::DrawStaticElements(FStaticPrimitiveDrawInter
 		BatchElement->MaxVertexIndex = SharedBuffers->IndexRanges[LOD].MaxIndexFull;
 	}
 
-	PDI->DrawMesh(MeshBatch,0,FLT_MAX);				
+	PDI->DrawMesh(MeshBatch, 0.0f);
 }
 
 uint64 FLandscapeVertexFactory::GetStaticBatchElementVisibility( const class FSceneView& View, const struct FMeshBatch* Batch ) const
diff --git a/Engine/Source/Runtime/Engine/Private/ModelRender.cpp b/Engine/Source/Runtime/Engine/Private/ModelRender.cpp
--- a/Engine/Source/Runtime/Engine/Private/ModelRender.cpp
+++ b/Engine/Source/Runtime/Engine/Private/ModelRender.cpp
@@ -453,7 +453,7 @@ public:
 					BatchElement.MaxVertexIndex = ModelElement.MaxVertexIndex;
 					MeshElement.Type = PT_TriangleList;
 					MeshElement.DepthPriorityGroup = PrimitiveDPG;
-					PDI->DrawMesh(MeshElement,0,FLT_MAX);
+					PDI->DrawMesh(MeshElement, 0.0f);
 				}
 			}
 		}
diff --git a/Engine/Source/Runtime/Engine/Private/SceneManagement.cpp b/Engine/Source/Runtime/Engine/Private/SceneManagement.cpp
--- a/Engine/Source/Runtime/Engine/Private/SceneManagement.cpp
+++ b/Engine/Source/Runtime/Engine/Private/SceneManagement.cpp
@@ -1,6 +1,7 @@
 // Copyright 1998-2014 Epic Games, Inc. All Rights Reserved.
 
 #include "EnginePrivate.h"
+#include "../../Renderer/Private/ScenePrivate.h"
 
 /*
 TUniformBufferRef<FPrimitiveUniformShaderParameters> CreatePrimitiveUniformBufferImmediate(
@@ -79,4 +80,76 @@ FLightMapInteraction FLightMapInteraction::Texture(
 	Result.CoordinateScale = InCoordinateScale;
 	Result.CoordinateBias = InCoordinateBias;
 	return Result;
-}
\ No newline at end of file
+}
+
+float ComputeBoundsScreenSize( const FVector4& Origin, const float SphereRadius, const FSceneView& View )
+{
+	// Only need one component from a view transformation; just calculate the one we're interested in.
+	const float Divisor =  Dot3(Origin - View.ViewMatrices.ViewOrigin, View.ViewMatrices.ViewMatrix.GetColumn(2));
+
+	// Get projection multiple accounting for view scaling.
+	const float ScreenMultiple = FMath::Max(View.ViewRect.Width() / 2.0f * View.ViewMatrices.ProjMatrix.M[0][0],
+		View.ViewRect.Height() / 2.0f * View.ViewMatrices.ProjMatrix.M[1][1]);
+
+	const float ScreenRadius = ScreenMultiple * SphereRadius / FMath::Max(Divisor, 1.0f);
+	const float ScreenArea = PI * ScreenRadius * ScreenRadius;
+	return FMath::Clamp(ScreenArea / View.ViewRect.Area(), 0.0f, 1.0f);
+}
+
+int8 ComputeStaticMeshLOD( const FStaticMeshRenderData* RenderData, const FVector4& Origin, const float SphereRadius, const FSceneView& View, float FactorScale )
+{
+	const int32 NumLODs = MAX_STATIC_MESH_LODS+1;
+
+	const float ScreenSize = ComputeBoundsScreenSize(Origin, SphereRadius, View) * FactorScale;
+
+	// Walk backwards and return the first matching LOD
+	for(int32 LODIndex = NumLODs - 1 ; LODIndex >= 0 ; --LODIndex)
+	{
+		if(RenderData->ScreenSize[LODIndex] > ScreenSize)
+		{
+			return LODIndex;
+		}
+	}
+
+	return 0;
+}
+
+int8 ComputeLODForMeshes( const TIndirectArray<class FStaticMesh>& StaticMeshes, FSceneView& View, const FVector4& Origin, float SphereRadius, int32 ForcedLODLevel, float ScreenSizeScale )
+{
+	int8 LODToRender = 0;
+
+	// Handle forced LOD level first
+	if(ForcedLODLevel >= 0)
+	{
+		int8 MaxLOD = 0;
+		for(int32 MeshIndex = 0 ; MeshIndex < StaticMeshes.Num() ; ++MeshIndex)
+		{
+			const FStaticMesh&  Mesh = StaticMeshes[MeshIndex];
+			MaxLOD = FMath::Max(MaxLOD, Mesh.LODIndex);
+		}
+		LODToRender = FMath::Clamp<int8>(ForcedLODLevel, 0, MaxLOD);
+	}
+	else
+	{
+		int32 NumMeshes = StaticMeshes.Num();
+
+		const float ScreenSize = ComputeBoundsScreenSize(Origin, SphereRadius, View);
+
+		for(int32 MeshIndex = NumMeshes-1 ; MeshIndex > 0 ; --MeshIndex)
+		{
+			const FStaticMesh& Mesh = StaticMeshes[MeshIndex];
+
+			if(View.Family->EngineShowFlags.LOD == 1)
+			{
+				float MeshScreenSize = Mesh.ScreenSize * ScreenSizeScale;
+
+				if(MeshScreenSize > ScreenSize)
+				{
+					LODToRender = Mesh.LODIndex;
+					break;
+				}
+			}
+		}
+	}
+	return LODToRender;
+}
diff --git a/Engine/Source/Runtime/Engine/Private/StaticMesh.cpp b/Engine/Source/Runtime/Engine/Private/StaticMesh.cpp
--- a/Engine/Source/Runtime/Engine/Private/StaticMesh.cpp
+++ b/Engine/Source/Runtime/Engine/Private/StaticMesh.cpp
@@ -530,7 +530,7 @@ FStaticMeshRenderData::FStaticMeshRenderData()
 {
 	for (int32 LODIndex = 0; LODIndex < MAX_STATIC_MESH_LODS+1; ++LODIndex)
 	{
-		LODDistance[LODIndex] = FLT_MAX;
+		ScreenSize[LODIndex] = 0.0f;
 	}
 
 	for (int32 TexCoordIndex = 0; TexCoordIndex < MAX_STATIC_TEXCOORDS; ++TexCoordIndex)
@@ -551,6 +551,7 @@ void FStaticMeshRenderData::Serialize(FArchive& Ar, UStaticMesh* Owner, bool bCo
 		Ar << WedgeMap;
 		Ar << MaterialIndexToImportIndex;
 	}
+
 #endif // #if WITH_EDITORONLY_DATA
 
 	LODResources.Serialize(Ar, Owner);
@@ -568,7 +569,33 @@ void FStaticMeshRenderData::Serialize(FArchive& Ar, UStaticMesh* Owner, bool bCo
 	{
 		for (int32 LODIndex = 0; LODIndex < MAX_STATIC_MESH_LODS+1; ++LODIndex)
 		{
-			Ar << LODDistance[LODIndex];
+			if(Ar.UE4Ver() < VER_UE4_STATIC_MESH_SCREEN_SIZE_LODS)
+			{
+				// We're loading in old data; which will be a LOD distance so we need
+				// to make an appropriate display factor from the distance;
+				float LODDistance = 0.0f;
+				Ar << LODDistance;
+
+				// Assuming an FOV of 90 and a screen size of 1920x1080 to estimate an appropriate display factor.
+				const float HalfFOV = PI / 4.0f;
+				const float ScreenWidth = 1920.0f;
+				const float ScreenHeight = 1080.0f;
+				const FVector4 PointToTest(0.0f, 0.0f, LODDistance, 1.0f);
+				FPerspectiveMatrix ProjMatrix(HalfFOV, ScreenWidth, ScreenHeight, 1.0f);
+
+				const float Divisor =  Dot3(-PointToTest, ProjMatrix.GetColumn(2));
+
+				FVector4 ScreenPosition = ProjMatrix.TransformFVector4(PointToTest);
+				const float ScreenMultiple = ScreenWidth / 2.0f * ProjMatrix.M[0][0];
+				const float ScreenRadius = ScreenMultiple * Bounds.SphereRadius / FMath::Max(ScreenPosition.W, 1.0f);
+				const float ScreenArea = ScreenWidth * ScreenHeight;
+				const float BoundsArea = PI * ScreenRadius * ScreenRadius;
+				ScreenSize[LODIndex] = FMath::Clamp(BoundsArea / ScreenArea, 0.0f, 1.0f);
+			}
+			else
+			{
+				Ar << ScreenSize[LODIndex];
+			}
 		}
 	}
 }
@@ -632,7 +659,8 @@ static float CalculateViewDistance(float MaxDeviation, float AllowedPixelError)
 void FStaticMeshRenderData::ResolveSectionInfo(UStaticMesh* Owner)
 {
 	int32 LODIndex = 0;
-	for (; LODIndex < LODResources.Num(); ++LODIndex)
+	int32 MaxLODs = LODResources.Num();
+	for (; LODIndex < MaxLODs; ++LODIndex)
 	{
 		FStaticMeshLODResources& LOD = LODResources[LODIndex];
 		for (int32 SectionIndex = 0; SectionIndex < LOD.Sections.Num(); ++SectionIndex)
@@ -646,34 +674,38 @@ void FStaticMeshRenderData::ResolveSectionInfo(UStaticMesh* Owner)
 
 		if (LODIndex == 0)
 		{
-			LODDistance[LODIndex] = 0.0f;
+			ScreenSize[LODIndex] = 1.0f;
 		}
-		else if (Owner->bAutoComputeLODDistance)
+		else if (Owner->bAutoComputeLODScreenSize)
 		{
-			if (LOD.MaxDeviation <= 0.0f)
+			if(LOD.MaxDeviation <= 0.0f)
 			{
-				const float MinAutoLODDistance = (LODIndex*2) * Bounds.SphereRadius;
-				LODDistance[LODIndex] = LODDistance[LODIndex-1] + MinAutoLODDistance;
+				ScreenSize[LODIndex] = 1.0f / (MaxLODs * LODIndex);
 			}
 			else
 			{
-				const float MinAutoLODDistance = 100.0f;
-				const float AutoLODDistance = CalculateViewDistance(LOD.MaxDeviation, Owner->AutoLODPixelError);
-				LODDistance[LODIndex] = FMath::Max<float>(AutoLODDistance, LODDistance[LODIndex-1] + MinAutoLODDistance);
+				const float ViewDistance = CalculateViewDistance(LOD.MaxDeviation, Owner->AutoLODPixelError);
+				ScreenSize[LODIndex] = 2.0f * Bounds.SphereRadius / ViewDistance;
 			}
 		}
 		else if (Owner->SourceModels.IsValidIndex(LODIndex))
 		{
-			LODDistance[LODIndex] = Owner->SourceModels[LODIndex].LODDistance;
+			ScreenSize[LODIndex] = Owner->SourceModels[LODIndex].ScreenSize;
 		}
 		else
 		{
-			LODDistance[LODIndex] = FLT_MAX;
+			// No valid source model and we're not auto-generating. Auto-generate in this case
+			// because we have nothing else to go on.
+			const float Tolerance = 0.01f;
+			float AutoDisplayFactor = 1.0f / (MaxLODs * LODIndex);
+
+			// Make sure this fits in with the previous LOD
+			ScreenSize[LODIndex] = FMath::Clamp(AutoDisplayFactor, 0.0f, ScreenSize[LODIndex-1] - Tolerance);
 		}
 	}
 	for (; LODIndex < MAX_STATIC_MESH_LODS + 1; ++LODIndex)
 	{
-		LODDistance[LODIndex] = FLT_MAX;
+		ScreenSize[LODIndex] = 0.0f;
 	}
 }
 
@@ -1062,7 +1094,7 @@ UStaticMesh::UStaticMesh(const FPostConstructInitializeProperties& PCIP)
 	bHasNavigationData=true;
 #if WITH_EDITORONLY_DATA
 	AutoLODPixelError = 1.0f;
-	bAutoComputeLODDistance=true;
+	bAutoComputeLODScreenSize=true;
 #endif // #if WITH_EDITORONLY_DATA
 	LpvBiasMultiplier = 1.0f;
 }
@@ -1280,14 +1312,14 @@ void UStaticMesh::PostEditChangeProperty(FPropertyChangedEvent& PropertyChangedE
 
 	AutoLODPixelError = FMath::Max(AutoLODPixelError, 1.0f);
 
-	if (!bAutoComputeLODDistance
+	if (!bAutoComputeLODScreenSize
 		&& RenderData
 		&& PropertyThatChanged
 		&& PropertyThatChanged->GetName() == TEXT("bAutoComputeLODDistance"))
 		{
 		for (int32 LODIndex = 1; LODIndex < SourceModels.Num(); ++LODIndex)
 		{
-			SourceModels[LODIndex].LODDistance = RenderData->LODDistance[LODIndex];
+			SourceModels[LODIndex].ScreenSize = RenderData->ScreenSize[LODIndex];
 		}
 	}
 
@@ -1364,7 +1396,7 @@ FStaticMeshSourceModel::FStaticMeshSourceModel()
 {
 #if WITH_EDITOR
 	RawMeshBulkData = new FRawMeshBulkData();
-	LODDistance = 0.0f;
+	ScreenSize = 0.0f;
 #endif // #if WITH_EDITOR
 }
 
@@ -1657,6 +1689,10 @@ void UStaticMesh::Serialize(FArchive& Ar)
 			SrcModel.SerializeBulkData(Ar, this);
 		}
 		SectionInfoMap.Serialize(Ar);
+
+		// Need to set a flag rather than do conversion in place as RenderData is not
+		// created until postload and it is needed for bounding information
+		bRequiresLODDistanceConversion = Ar.UE4Ver() < VER_UE4_STATIC_MESH_SCREEN_SIZE_LODS;
 	}
 #endif // #if WITH_EDITORONLY_DATA
 
@@ -1669,6 +1705,7 @@ void UStaticMesh::Serialize(FArchive& Ar)
 			RenderData = new FStaticMeshRenderData();
 			RenderData->Serialize(Ar, this, bCooked);
 		}
+
 #if WITH_EDITORONLY_DATA
 		else if (Ar.IsSaving())
 		{
@@ -1733,6 +1770,13 @@ void UStaticMesh::PostLoad()
 
 	CacheDerivedData();
 
+	// Only required in an editor build as other builds process this in a different place
+	if(bRequiresLODDistanceConversion)
+	{
+		// Convert distances to Display Factors
+		ConvertLegacyLODDistance();
+	}
+
 	if(RenderData && GStaticMeshesThatNeedMaterialFixup.Get(this))
 	{
 		FixupZeroTriangleSections();
@@ -2357,7 +2401,49 @@ UStaticMeshSocket* UStaticMesh::FindSocket(FName InSocketName)
 	return NULL;
 }
 
+void UStaticMesh::ConvertLegacyLODDistance()
+{
+	check(SourceModels.Num() > 0);
+	
+	if(SourceModels.Num() == 1)
+	{
+		// Only one model, 
+		SourceModels[0].ScreenSize = 1.0f;
+	}
+	else
+	{
+		// Multiple models, we should have LOD distance data.
+		// Assuming an FOV of 90 and a screen size of 1920x1080 to estimate an appropriate display factor.
+		const float HalfFOV = PI / 4.0f;
+		const float ScreenWidth = 1920.0f;
+		const float ScreenHeight = 1080.0f;
+
+		for(int32 ModelIndex = 0 ; ModelIndex < SourceModels.Num() ; ++ModelIndex)
+		{
+			FStaticMeshSourceModel& SrcModel = SourceModels[ModelIndex];
 
+			if(SrcModel.LODDistance_DEPRECATED == 0.0f)
+			{
+				SrcModel.ScreenSize = 1.0f;
+				RenderData->ScreenSize[ModelIndex] = SrcModel.ScreenSize;
+			}
+			else
+			{
+				// Create a screen position from the LOD distance
+				const FVector4 PointToTest(0.0f, 0.0f, SrcModel.LODDistance_DEPRECATED, 1.0f);
+				FPerspectiveMatrix ProjMatrix(HalfFOV, ScreenWidth, ScreenHeight, 1.0f);
+				FVector4 ScreenPosition = ProjMatrix.TransformFVector4(PointToTest);
+				// Convert to a percentage of the screen
+				const float ScreenMultiple = ScreenWidth / 2.0f * ProjMatrix.M[0][0];
+				const float ScreenRadius = ScreenMultiple * GetBounds().SphereRadius / FMath::Max(ScreenPosition.W, 1.0f);
+				const float ScreenArea = ScreenWidth * ScreenHeight;
+				const float BoundsArea = PI * ScreenRadius * ScreenRadius;
+				SrcModel.ScreenSize = FMath::Clamp(BoundsArea / ScreenArea, 0.0f, 1.0f);
+				RenderData->ScreenSize[ModelIndex] = SrcModel.ScreenSize;
+			}
+		}
+	}
+}
 
 /*-----------------------------------------------------------------------------
 UStaticMeshSocket
diff --git a/Engine/Source/Runtime/Engine/Private/StaticMeshRender.cpp b/Engine/Source/Runtime/Engine/Private/StaticMeshRender.cpp
--- a/Engine/Source/Runtime/Engine/Private/StaticMeshRender.cpp
+++ b/Engine/Source/Runtime/Engine/Private/StaticMeshRender.cpp
@@ -372,7 +372,7 @@ void FStaticMeshSceneProxy::DrawStaticElements(FStaticPrimitiveDrawInterface* PD
 				FMeshBatch MeshElement;
 				if(GetMeshElement(LODIndex, SectionIndex, PrimitiveDPG, MeshElement, bUseSelectedMaterial, bUseHoveredMaterial))
 				{
-					PDI->DrawMesh(MeshElement, 0, FLT_MAX);
+					PDI->DrawMesh(MeshElement, 0.0f);
 				}
 			}
 		} 
@@ -381,8 +381,7 @@ void FStaticMeshSceneProxy::DrawStaticElements(FStaticPrimitiveDrawInterface* PD
 			for(int32 LODIndex = 0; LODIndex < NumLODs; LODIndex++)
 			{
 				const FStaticMeshLODResources& LODModel = RenderData->LODResources[LODIndex];
-				float MinDist = GetMinLODDist(LODIndex);
-				float MaxDist = GetMaxLODDist(LODIndex);
+				float ScreenSize = GetScreenSize(LODIndex);
 
 				bool bHaveShadowOnlyMesh = false;
 				if (GUseShadowIndexBuffer
@@ -414,7 +413,7 @@ void FStaticMeshSceneProxy::DrawStaticElements(FStaticPrimitiveDrawInterface* PD
 						if (GetShadowMeshElement(LODIndex, PrimitiveDPG, MeshElement))
 						{
 							bHaveShadowOnlyMesh = true;
-							PDI->DrawMesh(MeshElement,MinDist,MaxDist,/*bShadowOnly=*/true);
+							PDI->DrawMesh(MeshElement, ScreenSize, /*bShadowOnly=*/true);
 						}
 					}
 				}
@@ -437,7 +436,7 @@ void FStaticMeshSceneProxy::DrawStaticElements(FStaticPrimitiveDrawInterface* PD
 					{
 						// If we have submitted an optimized shadow-only mesh, remaining mesh elements must not cast shadows.
 						MeshElement.CastShadow = MeshElement.CastShadow && !bHaveShadowOnlyMesh;
-						PDI->DrawMesh(MeshElement, MinDist, MaxDist);
+						PDI->DrawMesh(MeshElement, ScreenSize);
 					}
 				}
 			}
@@ -999,14 +998,9 @@ FLightInteraction FStaticMeshSceneProxy::FLODInfo::GetInteraction(const FLightSc
 	return FLightInteraction::Dynamic();
 }
 
-float FStaticMeshSceneProxy::GetMinLODDist(int32 LODIndex) const 
+float FStaticMeshSceneProxy::GetScreenSize( int32 LODIndex ) const
 {
-	return RenderData->LODDistance[LODIndex];
-}
-
-float FStaticMeshSceneProxy::GetMaxLODDist(int32 LODIndex) const 
-{
-	return RenderData->LODDistance[LODIndex+1];
+	return RenderData->ScreenSize[LODIndex];
 }
 
 /**
@@ -1036,29 +1030,8 @@ int32 FStaticMeshSceneProxy::GetLOD(const FSceneView* View) const
 	}
 #endif
 
-	// Note: These distance calculations must match up with the main renderer!
-#if !WITH_EDITOR
-	const FVector ViewOriginForDistance = View->ViewMatrices.ViewOrigin;
-#else
-	const FVector ViewOriginForDistance = View->IsPerspectiveProjection() ? View->ViewMatrices.ViewOrigin : View->OverrideLODViewOrigin;
-#endif
-
-	float DistanceSquared = (GetBounds().Origin - ViewOriginForDistance).SizeSquared();
-
-	for(int32 LODIndex = LODs.Num() - 1; LODIndex >= 0; LODIndex--)
-	{
-		// Use the same distances as FStaticMeshSceneProxy::DrawStaticElements 
-		// To ensure that LODs change the same way when drawn in a static draw list or when rendered through DrawDynamicElements
-		const float MinDist = GetMinLODDist(LODIndex);
-		const float MaxDist = GetMaxLODDist(LODIndex);
-		
-		const float LODFactorDistanceSquared = DistanceSquared * FMath::Square(View->LODDistanceFactor);
-		if (LODFactorDistanceSquared >= FMath::Square(MinDist) && LODFactorDistanceSquared < FMath::Square(MaxDist))
-		{
-			return LODIndex;
-		}
-	}
-	return INDEX_NONE;
+	const FBoxSphereBounds& Bounds = GetBounds();
+	return ComputeStaticMeshLOD(RenderData, Bounds.Origin, Bounds.SphereRadius, *View);
 }
 
 FPrimitiveSceneProxy* UStaticMeshComponent::CreateSceneProxy()
diff --git a/Engine/Source/Runtime/Engine/Private/TextRenderComponent.cpp b/Engine/Source/Runtime/Engine/Private/TextRenderComponent.cpp
--- a/Engine/Source/Runtime/Engine/Private/TextRenderComponent.cpp
+++ b/Engine/Source/Runtime/Engine/Private/TextRenderComponent.cpp
@@ -504,7 +504,7 @@ void FTextRenderSceneProxy::DrawStaticElements(FStaticPrimitiveDrawInterface* PD
 		Mesh.bDisableBackfaceCulling = false;
 		Mesh.Type = PT_TriangleList;
 		Mesh.DepthPriorityGroup = SDPG_World;
-		PDI->DrawMesh(Mesh,0.0, FLT_MAX);
+		PDI->DrawMesh(Mesh, 1.0f);
 	}
 }
 
diff --git a/Engine/Source/Runtime/Engine/Public/SceneManagement.h b/Engine/Source/Runtime/Engine/Public/SceneManagement.h
--- a/Engine/Source/Runtime/Engine/Public/SceneManagement.h
+++ b/Engine/Source/Runtime/Engine/Public/SceneManagement.h
@@ -1512,8 +1512,7 @@ public:
 	virtual void SetHitProxy(HHitProxy* HitProxy) = 0;
 	virtual void DrawMesh(
 		const FMeshBatch& Mesh,
-		float MinDrawDistance,
-		float MaxDrawDistance,
+		float ScreenSize,
 		bool bShadowOnly = false
 		) = 0;
 };
@@ -3073,3 +3072,29 @@ extern ENGINE_API void DrawUVs(FViewport* InViewport, FCanvas* InCanvas, int32 I
 
 /** Returns true if the Material and Vertex Factory combination require adjacency information. */
 bool RequiresAdjacencyInformation( class UMaterialInterface* Material, const class FVertexFactoryType* VertexFactoryType );
+/**
+ * Computes the screen size of a given sphere bounds in the given view
+ * @param Origin - Origin of the bounds in world space
+ * @param SphereRadius - Radius of the sphere to use to calculate screen coverage
+ * @param View - The view to calculate the display factor for
+ * @return float - The screen size calculated
+ */
+float ENGINE_API ComputeBoundsScreenSize(const FVector4& Origin, const float SphereRadius, const FSceneView& View);
+
+/**
+ * Computes the LOD level for the given static meshes render data in the given view.
+ * @param RenderData - Render data for the mesh
+ * @param Origin - Origin of the bounds of the mesh in world space
+ * @param SphereRadius - Radius of the sphere to use to calculate screen coverage
+ * @param View - The view to calculate the LOD level for
+ */
+int8 ENGINE_API ComputeStaticMeshLOD(const FStaticMeshRenderData* RenderData, const FVector4& Origin, const float SphereRadius, const FSceneView& View, float FactorScale = 1.0f);
+
+/**
+ * Computes the LOD to render for the list of static meshes in the given view.
+ * @param StaticMeshes - List of static meshes.
+ * @param View - The view to render the LOD level for 
+ * @param Origin - Origin of the bounds of the mesh in world space
+ * @param SphereRadius - Radius of the sphere to use to calculate screen coverage
+ */
+int8 ENGINE_API ComputeLODForMeshes(const TIndirectArray<class FStaticMesh>& StaticMeshes, FSceneView& View, const FVector4& Origin, float SphereRadius, int32 ForcedLODLevel, float ScreenSizeScale = 1.0f);
diff --git a/Engine/Source/Runtime/Engine/Public/StaticMeshResources.h b/Engine/Source/Runtime/Engine/Public/StaticMeshResources.h
--- a/Engine/Source/Runtime/Engine/Public/StaticMeshResources.h
+++ b/Engine/Source/Runtime/Engine/Public/StaticMeshResources.h
@@ -596,8 +596,8 @@ public:
 	/** Per-LOD resources. */
 	TIndirectArray<FStaticMeshLODResources> LODResources;
 
-	/** Distance at which to switch LODs. */
-	float LODDistance[MAX_STATIC_MESH_LODS+1];
+	/** Screen size to switch LODs */
+	float ScreenSize[MAX_STATIC_MESH_LODS];
 
 	/** Streaming texture factors. */
 	float StreamingTextureFactors[MAX_STATIC_TEXCOORDS];
@@ -888,6 +888,13 @@ protected:
 	 * @param CurrentLevel - the LOD to find the max distance for
 	 */
 	float GetMaxLODDist(int32 CurrentLevel) const;
+
+	/**
+	 * Returns the display factor for the given LOD level
+	 *
+	 * @Param LODIndex - The LOD to get the display factor for
+	 */
+	float GetScreenSize(int32 LODIndex) const;
 };
 
 #if WITH_EDITORONLY_DATA
diff --git a/Engine/Source/Runtime/Renderer/Private/PrimitiveSceneInfo.cpp b/Engine/Source/Runtime/Renderer/Private/PrimitiveSceneInfo.cpp
--- a/Engine/Source/Runtime/Renderer/Private/PrimitiveSceneInfo.cpp
+++ b/Engine/Source/Runtime/Renderer/Private/PrimitiveSceneInfo.cpp
@@ -38,8 +38,7 @@ public:
 	}
 	virtual void DrawMesh(
 		const FMeshBatch& Mesh,
-		float MinDrawDistance,
-		float MaxDrawDistance,
+		float ScreenSize,
 		bool bShadowOnly
 		)
 	{
@@ -52,8 +51,7 @@ public:
 		FStaticMesh* StaticMesh = new(PrimitiveSceneInfo->StaticMeshes) FStaticMesh(
 			PrimitiveSceneInfo,
 			Mesh,
-			FMath::Square(FMath::Max(0.0f,MinDrawDistance)),
-			FMath::Square(FMath::Max(0.0f,MaxDrawDistance)),
+			ScreenSize,
 			bShadowOnly,
 			CurrentHitProxy ? CurrentHitProxy->Id : FHitProxyId()
 			);
diff --git a/Engine/Source/Runtime/Renderer/Private/RendererScene.cpp b/Engine/Source/Runtime/Renderer/Private/RendererScene.cpp
--- a/Engine/Source/Runtime/Renderer/Private/RendererScene.cpp
+++ b/Engine/Source/Runtime/Renderer/Private/RendererScene.cpp
@@ -1359,12 +1359,6 @@ void FScene::UpdateSpeedTreeWind(double CurrentTime)
 				SET_SPEEDTREE_TABLE_FLOAT4V(WindLeaf2Roll, SH_LEAF_2_ROLL_MAX_SCALE);
 				SET_SPEEDTREE_TABLE_FLOAT4V(WindFrondRipple, SH_FROND_RIPPLE_TIME);
 
-				// @todo This needs to be broken out into its own constant buffer so it can change per instance
-				// Right now smooth SpeedTree LOD only works with the LOD distances set on the UStaticMesh
-				UniformParameters.LodInfo.Set((float)StaticMesh->RenderData->LODResources.Num(), 
-												1.0f / StaticMesh->RenderData->LODDistance[StaticMesh->RenderData->LODResources.Num() - 1], 
-												0.0f, 0.0f);
-
 				//BeginSetUniformBufferContents(WindComputation->UniformBuffer, UniformParameters);
 				WindComputation->UniformBuffer.SetContents(UniformParameters);
 			}
diff --git a/Engine/Source/Runtime/Renderer/Private/SceneCore.h b/Engine/Source/Runtime/Renderer/Private/SceneCore.h
--- a/Engine/Source/Runtime/Renderer/Private/SceneCore.h
+++ b/Engine/Source/Runtime/Renderer/Private/SceneCore.h
@@ -123,11 +123,8 @@ public:
 		virtual void Remove() = 0;
 	};
 
-	/** The squared minimum distance to draw the primitive at. */
-	float MinDrawDistanceSquared;
-
-	/** The squared maximum distance to draw the primitive at. */
-	float MaxDrawDistanceSquared;
+	/** The screen space size to draw this primitive at */
+	float ScreenSize;
 
 	/** The render info for the primitive which created this mesh. */
 	FPrimitiveSceneInfo* PrimitiveSceneInfo;
@@ -145,14 +142,12 @@ public:
 	FStaticMesh(
 		FPrimitiveSceneInfo* InPrimitiveSceneInfo,
 		const FMeshBatch& InMesh,
-		float InMinDrawDistanceSquared,
-		float InMaxDrawDistanceSquared,
+		float InScreenSize,
 		bool bInShadowOnly,
 		FHitProxyId InHitProxyId
 		):
 		FMeshBatch(InMesh),
-		MinDrawDistanceSquared(InMinDrawDistanceSquared),
-		MaxDrawDistanceSquared(InMaxDrawDistanceSquared),
+		ScreenSize(InScreenSize),
 		PrimitiveSceneInfo(InPrimitiveSceneInfo),
 		HitProxyId(InHitProxyId),
 		Id(INDEX_NONE),
@@ -182,8 +177,7 @@ private:
 	/** Private copy constructor. */
 	FStaticMesh(const FStaticMesh& InStaticMesh):
 		FMeshBatch(InStaticMesh),
-		MinDrawDistanceSquared(InStaticMesh.MinDrawDistanceSquared),
-		MaxDrawDistanceSquared(InStaticMesh.MaxDrawDistanceSquared),
+		ScreenSize(InStaticMesh.ScreenSize),
 		PrimitiveSceneInfo(InStaticMesh.PrimitiveSceneInfo),
 		HitProxyId(InStaticMesh.HitProxyId),
 		Id(InStaticMesh.Id)
diff --git a/Engine/Source/Runtime/Renderer/Private/ScenePrivate.h b/Engine/Source/Runtime/Renderer/Private/ScenePrivate.h
--- a/Engine/Source/Runtime/Renderer/Private/ScenePrivate.h
+++ b/Engine/Source/Runtime/Renderer/Private/ScenePrivate.h
@@ -1541,18 +1541,4 @@ private:
 	bool bHasSkyLight;
 };
 
-/**
- * Computes the LOD to render for the list of static meshes.
- * @param List of static meshes.
- * @param The squared distance to the primitive.
- * @param Squared scale to apply to the squared max draw distance.
- * @param ForcedLODLevel from "r.ForceLOD"
- */
-int8 ComputeLODForMeshes(
-	const TIndirectArray<class FStaticMesh>& StaticMeshes,
-	float DistanceSquared,
-	float MaxDrawDistanceScaleSquared,
-	int32 ForcedLODLevel
-	);
-
 #endif // __SCENEPRIVATE_H__
diff --git a/Engine/Source/Runtime/Renderer/Private/SceneVisibility.cpp b/Engine/Source/Runtime/Renderer/Private/SceneVisibility.cpp
--- a/Engine/Source/Runtime/Renderer/Private/SceneVisibility.cpp
+++ b/Engine/Source/Runtime/Renderer/Private/SceneVisibility.cpp
@@ -90,42 +90,6 @@ static FAutoConsoleVariableRef CVarDistanceFadeMaxTravel( TEXT("r.DistanceFadeMa
 	Visibility determination.
 ------------------------------------------------------------------------------*/
 
-int8 ComputeLODForMeshes(
-	const TIndirectArray<FStaticMesh>& StaticMeshes,
-	float DistanceSquared,
-	float MaxDrawDistanceScaleSquared,
-	int32 ForcedLODLevel
-	)
-{
-	int8 LODToRender = INDEX_NONE;
-	int32 NumStaticMeshes = StaticMeshes.Num();
-
-	if (ForcedLODLevel >= 0)
-	{
-		int8 MaxLOD = 0;
-		for(int32 MeshIndex = 0;MeshIndex < NumStaticMeshes;MeshIndex++)
-		{
-			const FStaticMesh& StaticMesh = StaticMeshes[MeshIndex];
-			MaxLOD = FMath::Max(MaxLOD, StaticMesh.LODIndex);
-		}
-		LODToRender = FMath::Clamp<int8>((int8)ForcedLODLevel, 0, MaxLOD);
-	}
-	else
-	{
-		for(int32 MeshIndex = 0;MeshIndex < NumStaticMeshes;MeshIndex++)
-		{
-			const FStaticMesh& StaticMesh = StaticMeshes[MeshIndex];
-			if (DistanceSquared >= StaticMesh.MinDrawDistanceSquared 
-				&& DistanceSquared < StaticMesh.MaxDrawDistanceSquared * MaxDrawDistanceScaleSquared)
-			{
-				LODToRender = StaticMesh.LODIndex;
-				break;
-			}
-		}
-	}
-	return LODToRender;
-}
-
 /**
  * Update a primitive's fading state.
  * @param FadingState - State to update.
@@ -830,18 +794,15 @@ static void MarkRelevantStaticMeshesForView(
 		const FPrimitiveSceneInfo * RESTRICT PrimitiveSceneInfo = Scene->Primitives[PrimitiveIndex];
 		const FPrimitiveBounds & Bounds = Scene->PrimitiveBounds[PrimitiveIndex];
 		const FPrimitiveViewRelevance & ViewRelevance = View.PrimitiveViewRelevanceMap[PrimitiveIndex];
+		
+		static const auto* StaticMeshLODDistanceScale = IConsoleManager::Get().FindTConsoleVariableDataFloat(TEXT("r.StaticMeshLODDistanceScale"));
+		check(StaticMeshLODDistanceScale);
+		const float LODScale = StaticMeshLODDistanceScale->GetValueOnRenderThread();
 
-		float DistanceSquared = (Bounds.Origin - ViewOrigin).SizeSquared();
-		const float LODFactorDistanceSquared = DistanceSquared * FMath::Square(View.LODDistanceFactor * InvLODScale);
-
-		// Go through the meshes and find the LOD to render
-		int8 LODToRender = ComputeLODForMeshes(
-			PrimitiveSceneInfo->StaticMeshes,
-			LODFactorDistanceSquared,
-			MaxDrawDistanceScaleSquared,
-			ForcedLODLevel
-			);
+		int8 LODToRender = ComputeLODForMeshes( PrimitiveSceneInfo->StaticMeshes, View, Bounds.Origin, Bounds.SphereRadius, ForcedLODLevel, LODScale);
 
+		float DistanceSquared = (Bounds.Origin - ViewOrigin).SizeSquared();
+		const float LODFactorDistanceSquared = DistanceSquared * FMath::Square(View.LODDistanceFactor * (1.0f / LODScale));
 		const bool bDrawShadowDepth = FMath::Square(Bounds.SphereRadius) > MinScreenRadiusForCSMDepthSquared * LODFactorDistanceSquared;
 		const bool bDrawDepthOnly = bForceEarlyZPass || FMath::Square(Bounds.SphereRadius) > GMinScreenRadiusForDepthPrepass * GMinScreenRadiusForDepthPrepass * LODFactorDistanceSquared;
 		
diff --git a/Engine/Source/Runtime/Renderer/Private/ShadowSetup.cpp b/Engine/Source/Runtime/Renderer/Private/ShadowSetup.cpp
--- a/Engine/Source/Runtime/Renderer/Private/ShadowSetup.cpp
+++ b/Engine/Source/Runtime/Renderer/Private/ShadowSetup.cpp
@@ -826,18 +826,10 @@ void FProjectedShadowInfo::AddSubjectPrimitive(FPrimitiveSceneInfo* PrimitiveSce
 						// The old LOD will continue to be drawn even though a different LOD would be chosen by distance.
 						else
 						{
-							float DistanceSquared = 0.0f;
-							if (CurrentView.IsPerspectiveProjection())
-							{	
-								// Calculate LOD in the same way it is done for other static mesh draw lists
-								DistanceSquared = PrimitiveSceneInfo->Proxy->GetBounds().GetBox().ComputeSquaredDistanceToPoint(CurrentView.ShadowViewMatrices.ViewOrigin);
-							}
-							DistanceSquared *= FMath::Square(CurrentView.LODDistanceFactor);
-
+							int8 LODToRender = 0;
 							int32 ForcedLODLevel = (CurrentView.Family->EngineShowFlags.LOD) ? GetCVarForceLOD() : 0;
 
 							// Add the primitive's static mesh elements to the draw lists.
-							int8 LODToRender;
 							if ( bReflectiveShadowmap) 
 							{
 								LODToRender = -CHAR_MAX;
@@ -849,13 +841,10 @@ void FProjectedShadowInfo::AddSubjectPrimitive(FPrimitiveSceneInfo* PrimitiveSce
 							}
 							else
 							{
-								// Add the primitive's static mesh elements to the draw lists.
-							    LODToRender = ComputeLODForMeshes(
-								PrimitiveSceneInfo->StaticMeshes,
-								DistanceSquared,
-								GetCachedScalabilityCVars().ViewDistanceScaleSquared,
-								ForcedLODLevel
-								);
+								FPrimitiveBounds PrimitiveBounds;
+								PrimitiveBounds.Origin = Bounds.Origin;
+								PrimitiveBounds.SphereRadius = Bounds.SphereRadius;
+								LODToRender = ComputeLODForMeshes(PrimitiveSceneInfo->StaticMeshes, CurrentView, PrimitiveBounds.Origin, PrimitiveBounds.SphereRadius, ForcedLODLevel);
 							}
 
 							for (int32 MeshIndex = 0; MeshIndex < PrimitiveSceneInfo->StaticMeshes.Num(); MeshIndex++)
