Set the UProperty for the Materials array so that the static lighting is not rebuilt.
#ttp 335786 Lighting: Assigning a Material in the Staticmesh editor unbuilds lighting on placed meshes
reviewer Thomas.Sarkanen

diff --git a/Engine/Source/Editor/StaticMeshEditor/Private/StaticMeshEditorTools.cpp b/Engine/Source/Editor/StaticMeshEditor/Private/StaticMeshEditorTools.cpp
--- a/Engine/Source/Editor/StaticMeshEditor/Private/StaticMeshEditorTools.cpp
+++ b/Engine/Source/Editor/StaticMeshEditor/Private/StaticMeshEditorTools.cpp
@@ -803,6 +803,13 @@ void FMeshSectionSettingsLayout::GetMaterials(IMaterialListBuilder& ListBuilder)
 void FMeshSectionSettingsLayout::OnMaterialChanged(UMaterialInterface* NewMaterial, UMaterialInterface* PrevMaterial, int32 SlotIndex, bool bReplaceAll)
 {
 	UStaticMesh& StaticMesh = GetStaticMesh();
+
+	// flag the property (Materials) we're modifying so that not all of the object is rebuilt.
+	UProperty* ChangedProperty = NULL;
+	ChangedProperty = FindField<UProperty>( UStaticMesh::StaticClass(), "Materials" );
+	check(ChangedProperty);
+	StaticMesh.PreEditChange(ChangedProperty);
+	
 	check(StaticMesh.RenderData);
 	FMeshSectionInfo Info = StaticMesh.SectionInfoMap.Get(LODIndex, SlotIndex);
 	if (LODIndex == 0)
@@ -826,7 +833,7 @@ void FMeshSectionSettingsLayout::OnMaterialChanged(UMaterialInterface* NewMateri
 			StaticMesh.Materials[Info.MaterialIndex] = NewMaterial;
 		}
 	}
-	CallPostEditChange();
+	CallPostEditChange(ChangedProperty);
 }
 
 TSharedRef<SWidget> FMeshSectionSettingsLayout::OnGenerateNameWidgetsForMaterial(UMaterialInterface* Material, int32 SlotIndex)
@@ -952,11 +959,19 @@ void FMeshSectionSettingsLayout::OnSectionSelectedChanged(ESlateCheckBoxState::T
 	}
 }
 
-void FMeshSectionSettingsLayout::CallPostEditChange()
+void FMeshSectionSettingsLayout::CallPostEditChange(UProperty* PropertyChanged/*=nullptr*/)
 {
 	UStaticMesh& StaticMesh = GetStaticMesh();
-	StaticMesh.Modify();
-	StaticMesh.PostEditChange();
+	if( PropertyChanged )
+	{
+		FPropertyChangedEvent PropertyUpdateStruct(PropertyChanged);
+		StaticMesh.PostEditChangeProperty(PropertyUpdateStruct);
+	}
+	else
+	{
+		StaticMesh.Modify();
+		StaticMesh.PostEditChange();
+	}
 	if(StaticMesh.BodySetup)
 	{
 		StaticMesh.BodySetup->CreatePhysicsMeshes();
diff --git a/Engine/Source/Editor/StaticMeshEditor/Private/StaticMeshEditorTools.h b/Engine/Source/Editor/StaticMeshEditor/Private/StaticMeshEditorTools.h
--- a/Engine/Source/Editor/StaticMeshEditor/Private/StaticMeshEditorTools.h
+++ b/Engine/Source/Editor/StaticMeshEditor/Private/StaticMeshEditorTools.h
@@ -225,7 +225,7 @@ private:
 	void OnSectionCollisionChanged(ESlateCheckBoxState::Type NewState, int32 SectionIndex);
 	ESlateCheckBoxState::Type IsSectionSelected(int32 SectionIndex) const;
 	void OnSectionSelectedChanged(ESlateCheckBoxState::Type NewState, int32 SectionIndex);
-	void CallPostEditChange();
+	void CallPostEditChange(UProperty* PropertyChanged=nullptr);
 	
 	IStaticMeshEditor& StaticMeshEditor;
 	int32 LODIndex;
