Static mesh editor can now create Box/Sphyl and have multiple collision primitives, which can be manipulated.

#ttp 334968 - TOOLS FEATURE: Add primitive collision inside the static mesh editor

#branch UE4

#change
GenerateSimpleCollision funcs now return true if any collision was built
Moved Prompt to build collision into a func (to remove duplicate code), and modified so that the options are Yes=Replace, No=Add-To, Cancel = Cancel.
Added CalcBoundingSphyl code which is loosely based on the Sphere code in that it finds the best axis to align the sphyl with (but locally rotates so it can treat the Z axis as length, regardless). It they finds the radius to best enclose the Sphyl (same as Sphere, but 2D). Then it increase the length of the Sphyl until that encapsulates all the remaining points.

Replace instances of AddZeroed with Add(FKTypeElem()) so that the vtable would be created correctly if a base struct was introduced.
Moved PhysX code which scales the Elems from ModifyPrimitiveSize to their own struct func (ScaleElem) so that it could be used elsewhere.

Added new struct FPrimData which contains the PrimType and PrimIndex so we can look the corresponding collision Elem up in our BodySetup. The static mesh editor session which is open for a particular static mesh keeps track of these using the following:
IsPrimValid û Checks to see if the prim data could be valid
HasSelectedPrims û Returns true if there are any selected prims
IsSelectedPrim û Check to see if the prim data is selected
AddSelectedPrim û Adds prim data to our selection
RemoveSelectedPrim û Removes prim data from our selection
RemoveInvalidPrism - Removes any prims which are no longer valid from the selection
ClearSelectedPrims û Removes all the prim data from our selection
DuplicateSelectedPrims û Duplicates the selected prims
TranslateSelectedPrims - Translates the selected prims by a specified amount
RotateSelectedPrims - Rotates the selected prims by a specified amount
ScaleSelectedPrims - Scales the selected prims by a specified amount
CalcSelectedPrimsAABB - Calculates the bounding box of the selected prims
DeleteSelectedPrims û Deletes the selected prims
GetLastSelectedPrimTransform û Gets the transform of the last prim which was added to the selection.
GetPrimTransform - Gets the transform of a specified prim
SetPrimTransform - Sets the transform of a specified prim
Added placeholder calls to begin/end transaction where itÆs needed whenever the bodysetup prims are modified.

Added new Static Mesh Editor menu option to create a capsule/sphyl
Added dedicated collision toggle flag to the Static Mesh Editor so we could manage our own handling of the draw functions in order to inject hit proxies where needed so we can test when theyÆve been clicked on in the viewport.
Widget now works when selecting collision Elems. Space toggles the manipulation method.
SelectedPrims can be trans/rot/scaled using the editor widget.
Prims can be duplicated or deleted either using menu entries or keyboard shortcuts. Alt+drag is also supported for duplication.
ProcessClick was updated so that whenever a new selection type is clicked on the previous types are deselected (sockets, prims, edges), and other code location were updated to clear the selected prims where appropriate.
Modified OnFocusViewportToSelection so that it also focuses on the selected elems too
Modified SetSelectedSocket so that when passing NULL it deselects all the sockets

diff --git a/Engine/Source/Editor/PhAT/Private/PhAT.cpp b/Engine/Source/Editor/PhAT/Private/PhAT.cpp
--- a/Engine/Source/Editor/PhAT/Private/PhAT.cpp
+++ b/Engine/Source/Editor/PhAT/Private/PhAT.cpp
@@ -1870,7 +1870,7 @@ void FPhAT::AddNewPrimitive(EKCollisionPrimitiveType InPrimitiveType, bool bCopy
 
 			if (PrimitiveType == KPT_Sphere)
 			{
-				NewPrimIndex = BodySetup->AggGeom.SphereElems.AddZeroed();
+				NewPrimIndex = BodySetup->AggGeom.SphereElems.Add(FKSphereElem());
 				NewSelection[i].PrimitiveType = KPT_Sphere;
 				NewSelection[i].PrimitiveIndex = NewPrimIndex;
 				FKSphereElem* SphereElem = &BodySetup->AggGeom.SphereElems[NewPrimIndex];
@@ -1891,7 +1891,7 @@ void FPhAT::AddNewPrimitive(EKCollisionPrimitiveType InPrimitiveType, bool bCopy
 			}
 			else if (PrimitiveType == KPT_Box)
 			{
-				NewPrimIndex = BodySetup->AggGeom.BoxElems.AddZeroed();
+				NewPrimIndex = BodySetup->AggGeom.BoxElems.Add(FKBoxElem());
 				NewSelection[i].PrimitiveType = KPT_Box;
 				NewSelection[i].PrimitiveIndex = NewPrimIndex;
 				FKBoxElem* BoxElem = &BodySetup->AggGeom.BoxElems[NewPrimIndex];
@@ -1916,7 +1916,7 @@ void FPhAT::AddNewPrimitive(EKCollisionPrimitiveType InPrimitiveType, bool bCopy
 			}
 			else if (PrimitiveType == KPT_Sphyl)
 			{
-				NewPrimIndex = BodySetup->AggGeom.SphylElems.AddZeroed();
+				NewPrimIndex = BodySetup->AggGeom.SphylElems.Add(FKSphylElem());
 				NewSelection[i].PrimitiveType = KPT_Sphyl;
 				NewSelection[i].PrimitiveIndex = NewPrimIndex;
 				FKSphylElem* SphylElem = &BodySetup->AggGeom.SphylElems[NewPrimIndex];
diff --git a/Engine/Source/Editor/PhAT/Private/PhATPreviewViewportClient.cpp b/Engine/Source/Editor/PhAT/Private/PhATPreviewViewportClient.cpp
--- a/Engine/Source/Editor/PhAT/Private/PhATPreviewViewportClient.cpp
+++ b/Engine/Source/Editor/PhAT/Private/PhATPreviewViewportClient.cpp
@@ -931,53 +931,30 @@ void FPhATEdPreviewViewportClient::ModifyPrimitiveSize(int32 BodyIndex, EKCollis
 
 	if (PrimType == KPT_Sphere)
 	{
-		// Find element with largest magnitude, btu preserve sign.
-		float DeltaRadius = DeltaSize.X;
-		if (FMath::Abs(DeltaSize.Y) > FMath::Abs(DeltaRadius))
-			DeltaRadius = DeltaSize.Y;
-		else if (FMath::Abs(DeltaSize.Z) > FMath::Abs(DeltaRadius))
-			DeltaRadius = DeltaSize.Z;
-
-		AggGeom->SphereElems[PrimIndex].Radius = FMath::Max(AggGeom->SphereElems[PrimIndex].Radius + DeltaRadius, MinPrimSize);
+		check(AggGeom->SphereElems.IsValidIndex(PrimIndex));
+		AggGeom->SphereElems[PrimIndex].ScaleElem(DeltaSize, MinPrimSize);
 	}
 	else if (PrimType == KPT_Box)
 	{
-		// Sizes are lengths, so we double the delta to get similar increase in size.
-		AggGeom->BoxElems[PrimIndex].X = FMath::Max(AggGeom->BoxElems[PrimIndex].X + 2*DeltaSize.X, MinPrimSize);
-		AggGeom->BoxElems[PrimIndex].Y = FMath::Max(AggGeom->BoxElems[PrimIndex].Y + 2*DeltaSize.Y, MinPrimSize);
-		AggGeom->BoxElems[PrimIndex].Z = FMath::Max(AggGeom->BoxElems[PrimIndex].Z + 2*DeltaSize.Z, MinPrimSize);
-
+		check(AggGeom->BoxElems.IsValidIndex(PrimIndex));
+		AggGeom->BoxElems[PrimIndex].ScaleElem(DeltaSize, MinPrimSize);
 	}
 	else if (PrimType == KPT_Sphyl)
 	{
-		FKSphylElem & SphylElem = AggGeom->SphylElems[PrimIndex];
-		float DeltaRadius = DeltaSize.X;
-		if (FMath::Abs(DeltaSize.Y) > FMath::Abs(DeltaRadius))
-		{
-			DeltaRadius = DeltaSize.Y;
-		}
-
-		float DeltaHeight = DeltaSize.Z;
-		float Radius = FMath::Max(SphylElem.Radius + DeltaRadius, MinPrimSize);
-		float Length = SphylElem.Length + DeltaHeight;
-		
-		Length += SphylElem.Radius - Radius;
-		Length = FMath::Max(0.f, Length);
-
-		SphylElem.Radius = Radius;
-		SphylElem.Length = Length;
+		check(AggGeom->SphylElems.IsValidIndex(PrimIndex));
+		AggGeom->SphylElems[PrimIndex].ScaleElem(DeltaSize, MinPrimSize);
 	}
 	else if (PrimType == KPT_Convex)
 	{
-		FVector ModifiedScale = DeltaSize;
-		if( GEditor->UsePercentageBasedScaling() )
+		check(AggGeom->ConvexElems.IsValidIndex(PrimIndex));
+
+		FVector ModifiedSize = DeltaSize;
+		if (GEditor->UsePercentageBasedScaling())
 		{
-			ModifiedScale = DeltaSize * ((GEditor->GetScaleGridSize() / 100.0f) / GEditor->GetGridSize());
+			ModifiedSize = DeltaSize * ((GEditor->GetScaleGridSize() / 100.0f) / GEditor->GetGridSize());
 		}
 
-		FTransform Transform = AggGeom->ConvexElems[PrimIndex].GetTransform();
-		Transform.SetScale3D(Transform.GetScale3D() + ModifiedScale);
-		AggGeom->ConvexElems[PrimIndex].SetTransform(Transform);
+		AggGeom->ConvexElems[PrimIndex].ScaleElem(ModifiedSize, MinPrimSize);
 	}
 }
 
diff --git a/Engine/Source/Editor/PhAT/Private/PhATSharedData.cpp b/Engine/Source/Editor/PhAT/Private/PhATSharedData.cpp
--- a/Engine/Source/Editor/PhAT/Private/PhATSharedData.cpp
+++ b/Engine/Source/Editor/PhAT/Private/PhATSharedData.cpp
@@ -124,13 +124,13 @@ void FPhATSharedData::Initialize()
 		UBodySetup* BodySetup = PhysicsAsset->BodySetup[i];
 		if (BodySetup->AggGeom.GetElementCount() == 0)
 		{
-			BodySetup->AggGeom.BoxElems.AddZeroed(1);
+			FKBoxElem BoxElem;
+			BoxElem.SetTransform(FTransform::Identity);
+			BoxElem.X = 15.f;
+			BoxElem.Y = 15.f;
+			BoxElem.Z = 15.f;
+			BodySetup->AggGeom.BoxElems.Add(BoxElem);
 			check(BodySetup->AggGeom.BoxElems.Num() == 1);
-			FKBoxElem& Box = BodySetup->AggGeom.BoxElems[0];
-			Box.SetTransform( FTransform::Identity );
-			Box.X = 15.f;
-			Box.Y = 15.f;
-			Box.Z = 15.f;
 
 			bFoundEmptyShape = true;
 		}
diff --git a/Engine/Source/Editor/StaticMeshEditor/Private/SStaticMeshEditorViewport.cpp b/Engine/Source/Editor/StaticMeshEditor/Private/SStaticMeshEditorViewport.cpp
--- a/Engine/Source/Editor/StaticMeshEditor/Private/SStaticMeshEditorViewport.cpp
+++ b/Engine/Source/Editor/StaticMeshEditor/Private/SStaticMeshEditorViewport.cpp
@@ -336,9 +336,9 @@ void SStaticMeshEditorViewport::BindCommands()
 
 	CommandList->MapAction(
 		Commands.SetShowCollision,
-		FExecuteAction::CreateSP( EditorViewportClientRef, &FStaticMeshEditorViewportClient::SetShowCollision ),
+		FExecuteAction::CreateSP( EditorViewportClientRef, &FStaticMeshEditorViewportClient::SetShowWireframeCollision ),
 		FCanExecuteAction(),
-		FIsActionChecked::CreateSP( EditorViewportClientRef, &FStaticMeshEditorViewportClient::IsSetShowCollisionChecked ) );
+		FIsActionChecked::CreateSP( EditorViewportClientRef, &FStaticMeshEditorViewportClient::IsSetShowWireframeCollisionChecked ) );
 
 	CommandList->MapAction(
 		Commands.SetShowSockets,
@@ -380,22 +380,35 @@ void SStaticMeshEditorViewport::BindCommands()
 
 void SStaticMeshEditorViewport::OnFocusViewportToSelection()
 {
+	// If we have selected sockets, focus on them
 	UStaticMeshSocket* SelectedSocket = StaticMeshEditorPtr.Pin()->GetSelectedSocket();
-
 	if( SelectedSocket && PreviewMeshComponent )
 	{
 		FTransform SocketTransform;
 		SelectedSocket->GetSocketTransform( SocketTransform, PreviewMeshComponent );
 
-		FVector Extent(30.0f);
+		const FVector Extent(30.0f);
+
+		const FVector Origin = SocketTransform.GetLocation();
+		const FBox Box(Origin - Extent, Origin + Extent);
 
-		FVector Origin = SocketTransform.GetLocation();
-		FBox Box( Origin - Extent, Origin + Extent);
-	
 		EditorViewportClient->FocusViewportOnBox( Box );
+		return;
+	}
+
+	// If we have selected primitives, focus on them 
+	FBox Box(0);
+	const bool bSelectedPrim = StaticMeshEditorPtr.Pin()->CalcSelectedPrimsAABB(Box);
+	if (bSelectedPrim)
+	{
+		EditorViewportClient->FocusViewportOnBox(Box);
+		return;
 	}
-	else if( PreviewMeshComponent )
+
+	// Fallback to focusing on the mesh, if nothing else
+	if( PreviewMeshComponent )
 	{
 		EditorViewportClient->FocusViewportOnBox( PreviewMeshComponent->Bounds.GetBox() );
+		return;
 	}
 }
\ No newline at end of file
diff --git a/Engine/Source/Editor/StaticMeshEditor/Private/StaticMeshEditor.cpp b/Engine/Source/Editor/StaticMeshEditor/Private/StaticMeshEditor.cpp
--- a/Engine/Source/Editor/StaticMeshEditor/Private/StaticMeshEditor.cpp
+++ b/Engine/Source/Editor/StaticMeshEditor/Private/StaticMeshEditor.cpp
@@ -215,16 +215,18 @@ void FStaticMeshEditor::ExtendMenu()
 		{
 			InMenuBuilder.BeginSection("CollisionEditCollision");
 			{
-				InMenuBuilder.AddMenuEntry(FStaticMeshEditorCommands::Get().CreateDOP6);
+				InMenuBuilder.AddMenuEntry(FStaticMeshEditorCommands::Get().CreateSphereCollision);
+				InMenuBuilder.AddMenuEntry(FStaticMeshEditorCommands::Get().CreateSphylCollision);
+				InMenuBuilder.AddMenuEntry(FStaticMeshEditorCommands::Get().CreateBoxCollision);
 				InMenuBuilder.AddMenuEntry(FStaticMeshEditorCommands::Get().CreateDOP10X);
 				InMenuBuilder.AddMenuEntry(FStaticMeshEditorCommands::Get().CreateDOP10Y);
 				InMenuBuilder.AddMenuEntry(FStaticMeshEditorCommands::Get().CreateDOP10Z);
 				InMenuBuilder.AddMenuEntry(FStaticMeshEditorCommands::Get().CreateDOP18);
 				InMenuBuilder.AddMenuEntry(FStaticMeshEditorCommands::Get().CreateDOP26);	
-
-				InMenuBuilder.AddMenuEntry(FStaticMeshEditorCommands::Get().CreateSphereCollision);
-				InMenuBuilder.AddMenuEntry(FStaticMeshEditorCommands::Get().RemoveCollision);
 				InMenuBuilder.AddMenuEntry(FStaticMeshEditorCommands::Get().ConvertBoxesToConvex);
+				InMenuBuilder.AddMenuEntry(FStaticMeshEditorCommands::Get().RemoveCollision);
+				InMenuBuilder.AddMenuEntry(FGenericCommands::Get().Delete, "DeleteCollision", LOCTEXT("DeleteCollision", "Delete Selected Collision"), LOCTEXT("DeleteCollisionToolTip", "Deletes the selected Collision from the mesh."));
+				InMenuBuilder.AddMenuEntry(FGenericCommands::Get().Duplicate, "DuplicateCollision", LOCTEXT("DuplicateCollision", "Duplicate Selected Collision"), LOCTEXT("DuplicateCollisionToolTip", "Duplicates the selected Collision."));
 			}
 			InMenuBuilder.EndSection();
 
@@ -361,8 +363,8 @@ void FStaticMeshEditor::BindCommands()
 	const TSharedRef<FUICommandList>& UICommandList = GetToolkitCommands();
 
 	UICommandList->MapAction( FGenericCommands::Get().Delete,
-		FExecuteAction::CreateSP( this, &FStaticMeshEditor::DeleteSelectedSockets ),
-		FCanExecuteAction::CreateSP( this, &FStaticMeshEditor::HasSelectedSockets ));
+		FExecuteAction::CreateSP( this, &FStaticMeshEditor::DeleteSelected ),
+		FCanExecuteAction::CreateSP(this, &FStaticMeshEditor::CanDeleteSelected));
 
 	UICommandList->MapAction( FGenericCommands::Get().Undo, 
 		FExecuteAction::CreateSP( this, &FStaticMeshEditor::UndoAction ) );
@@ -372,17 +374,13 @@ void FStaticMeshEditor::BindCommands()
 
 	UICommandList->MapAction(
 		FGenericCommands::Get().Duplicate,
-		FExecuteAction::CreateSP(this, &FStaticMeshEditor::DuplicateSelectedSocket),
-		FCanExecuteAction::CreateSP(this, &FStaticMeshEditor::HasSelectedSockets));
+		FExecuteAction::CreateSP(this, &FStaticMeshEditor::DuplicateSelected),
+		FCanExecuteAction::CreateSP(this, &FStaticMeshEditor::CanDuplicateSelected));
 
 	UICommandList->MapAction(
 		FGenericCommands::Get().Rename,
 		FExecuteAction::CreateSP(this, &FStaticMeshEditor::RequestRenameSelectedSocket),
-		FCanExecuteAction::CreateSP(this, &FStaticMeshEditor::HasSelectedSockets));
-
-	UICommandList->MapAction(
-		Commands.CreateDOP6,
-		FExecuteAction::CreateSP(this, &FStaticMeshEditor::GenerateKDop, KDopDir6, (uint32)6));
+		FCanExecuteAction::CreateSP(this, &FStaticMeshEditor::CanRenameSelected));
 
 	UICommandList->MapAction(
 		Commands.CreateDOP10X,
@@ -404,10 +402,18 @@ void FStaticMeshEditor::BindCommands()
 		Commands.CreateDOP26,
 		FExecuteAction::CreateSP(this, &FStaticMeshEditor::GenerateKDop, KDopDir26, (uint32)26));
 
+	UICommandList->MapAction(
+		Commands.CreateBoxCollision,
+		FExecuteAction::CreateSP(this, &FStaticMeshEditor::OnCollisionBox));
+
 	UICommandList->MapAction(
 		Commands.CreateSphereCollision,
 		FExecuteAction::CreateSP(this, &FStaticMeshEditor::OnCollisionSphere));
 
+	UICommandList->MapAction(
+		Commands.CreateSphylCollision,
+		FExecuteAction::CreateSP(this, &FStaticMeshEditor::OnCollisionSphyl));
+
 	UICommandList->MapAction(
 		Commands.RemoveCollision,
 		FExecuteAction::CreateSP(this, &FStaticMeshEditor::OnRemoveCollision));
@@ -581,6 +587,323 @@ void FStaticMeshEditor::RequestRenameSelectedSocket()
 	SocketManager->RequestRenameSelectedSocket();
 }
 
+bool FStaticMeshEditor::IsPrimValid(const FPrimData& InPrimData) const
+{
+	if (StaticMesh->BodySetup)
+	{
+		const FKAggregateGeom* AggGeom = &StaticMesh->BodySetup->AggGeom;
+
+		switch (InPrimData.PrimType)
+		{
+		case KPT_Sphere:
+			return AggGeom->SphereElems.IsValidIndex(InPrimData.PrimIndex);
+		case KPT_Box:
+			return AggGeom->BoxElems.IsValidIndex(InPrimData.PrimIndex);
+		case KPT_Sphyl:
+			return AggGeom->SphylElems.IsValidIndex(InPrimData.PrimIndex);
+		case KPT_Convex:
+			return AggGeom->ConvexElems.IsValidIndex(InPrimData.PrimIndex);
+		}
+	}
+	return false;
+}
+
+bool FStaticMeshEditor::HasSelectedPrims() const
+{
+	return (SelectedPrims.Num() > 0 ? true : false);
+}
+
+void FStaticMeshEditor::AddSelectedPrim(const FPrimData& InPrimData)
+{
+	check(IsPrimValid(InPrimData));
+	SelectedPrims.Add(InPrimData);	
+}
+
+void FStaticMeshEditor::RemoveSelectedPrim(const FPrimData& InPrimData)
+{
+	SelectedPrims.Remove(InPrimData);
+}
+
+void FStaticMeshEditor::RemoveInvalidPrims()
+{
+	for (int32 PrimIdx = SelectedPrims.Num() - 1; PrimIdx >= 0; PrimIdx--)
+	{
+		FPrimData& PrimData = SelectedPrims[PrimIdx];
+
+		if (!IsPrimValid(PrimData))
+		{
+			SelectedPrims.RemoveAt(PrimIdx);
+		}
+	}
+}
+
+bool FStaticMeshEditor::IsSelectedPrim(const FPrimData& InPrimData) const
+{
+	return SelectedPrims.Contains(InPrimData);
+}
+
+void FStaticMeshEditor::ClearSelectedPrims()
+{
+	SelectedPrims.Empty();
+}
+
+void FStaticMeshEditor::DuplicateSelectedPrims(const FVector* InOffset)
+{
+	if (SelectedPrims.Num() > 0)
+	{
+		check(StaticMesh->BodySetup);
+
+		FKAggregateGeom* AggGeom = &StaticMesh->BodySetup->AggGeom;
+
+		GEditor->BeginTransaction(LOCTEXT("FStaticMeshEditor_DuplicateSelectedPrims", "Duplicate Collision"));
+		StaticMesh->BodySetup->Modify();
+
+		//Clear the cache (PIE may have created some data), create new GUID
+		StaticMesh->BodySetup->InvalidatePhysicsData();
+
+		for (int32 PrimIdx = 0; PrimIdx < SelectedPrims.Num(); PrimIdx++)
+		{
+			FPrimData& PrimData = SelectedPrims[PrimIdx];
+
+			check(IsPrimValid(PrimData));
+			switch (PrimData.PrimType)
+			{
+			case KPT_Sphere:
+				{					
+					const FKSphereElem SphereElem = AggGeom->SphereElems[PrimData.PrimIndex];
+					PrimData.PrimIndex = AggGeom->SphereElems.Add(SphereElem);
+				}
+				break;
+			case KPT_Box:
+				{
+					const FKBoxElem BoxElem = AggGeom->BoxElems[PrimData.PrimIndex];
+					PrimData.PrimIndex = AggGeom->BoxElems.Add(BoxElem);
+				}
+				break;
+			case KPT_Sphyl:
+				{
+					const FKSphylElem SphylElem = AggGeom->SphylElems[PrimData.PrimIndex];
+					PrimData.PrimIndex = AggGeom->SphylElems.Add(SphylElem);
+				}
+				break;
+			case KPT_Convex:
+				{
+					const FKConvexElem ConvexElem = AggGeom->ConvexElems[PrimData.PrimIndex];
+					PrimData.PrimIndex = AggGeom->ConvexElems.Add(ConvexElem);
+				}
+				break;
+			}
+
+			// If specified, offset the duplicate by a specific amount
+			if (InOffset)
+			{
+				FTransform PrimTransform = GetPrimTransform(PrimData);
+				FVector PrimLocation = PrimTransform.GetLocation();
+				PrimLocation += *InOffset;
+				PrimTransform.SetLocation(PrimLocation);
+				SetPrimTransform(PrimData, PrimTransform);
+			}
+		}
+
+		// refresh collision change back to staticmesh components
+		RefreshCollisionChange(StaticMesh);
+
+		GEditor->EndTransaction();
+
+		// Mark staticmesh as dirty, to help make sure it gets saved.
+		StaticMesh->MarkPackageDirty();
+
+		// Update views/property windows
+		Viewport->RefreshViewport();
+	}
+}
+
+void FStaticMeshEditor::TranslateSelectedPrims(const FVector& InDrag)
+{
+	check(StaticMesh->BodySetup);
+
+	FKAggregateGeom* AggGeom = &StaticMesh->BodySetup->AggGeom;
+
+	for (int32 PrimIdx = 0; PrimIdx < SelectedPrims.Num(); PrimIdx++)
+	{
+		const FPrimData& PrimData = SelectedPrims[PrimIdx];
+
+		FTransform PrimTransform = GetPrimTransform(PrimData);
+
+		FVector PrimLocation = PrimTransform.GetLocation();
+		PrimLocation += InDrag;
+		PrimTransform.SetLocation(PrimLocation);
+
+		SetPrimTransform(PrimData, PrimTransform);
+	}
+}
+
+void FStaticMeshEditor::RotateSelectedPrims(const FRotator& InRot)
+{
+	check(StaticMesh->BodySetup);
+
+	FKAggregateGeom* AggGeom = &StaticMesh->BodySetup->AggGeom;
+
+	const FQuat DeltaQ = InRot.Quaternion();
+
+	for (int32 PrimIdx = 0; PrimIdx < SelectedPrims.Num(); PrimIdx++)
+	{
+		const FPrimData& PrimData = SelectedPrims[PrimIdx];
+
+		FTransform PrimTransform = GetPrimTransform(PrimData);
+
+		FRotator ActorRotWind, ActorRotRem;
+		PrimTransform.Rotator().GetWindingAndRemainder(ActorRotWind, ActorRotRem);
+
+		const FQuat ActorQ = ActorRotRem.Quaternion();
+		FRotator NewActorRotRem = FRotator(DeltaQ * ActorQ);
+		NewActorRotRem.Normalize();
+		PrimTransform.SetRotation(NewActorRotRem.Quaternion());
+
+		SetPrimTransform(PrimData, PrimTransform);
+	}
+}
+
+void FStaticMeshEditor::ScaleSelectedPrims(const FVector& InScale)
+{
+	check(StaticMesh->BodySetup);
+
+	FKAggregateGeom* AggGeom = &StaticMesh->BodySetup->AggGeom;
+
+	FVector ModifiedScale = InScale;
+	if (GEditor->UsePercentageBasedScaling())
+	{
+		ModifiedScale = InScale * ((GEditor->GetScaleGridSize() / 100.0f) / GEditor->GetGridSize());
+	}
+
+	for (int32 PrimIdx = 0; PrimIdx < SelectedPrims.Num(); PrimIdx++)
+	{
+		const FPrimData& PrimData = SelectedPrims[PrimIdx];
+
+		check(IsPrimValid(PrimData));
+		switch (PrimData.PrimType)
+		{
+		case KPT_Sphere:
+			AggGeom->SphereElems[PrimData.PrimIndex].ScaleElem(InScale, MinPrimSize);
+			break;
+		case KPT_Box:
+			AggGeom->BoxElems[PrimData.PrimIndex].ScaleElem(InScale, MinPrimSize);
+			break;
+		case KPT_Sphyl:
+			AggGeom->SphylElems[PrimData.PrimIndex].ScaleElem(InScale, MinPrimSize);
+			break;
+		case KPT_Convex:
+			AggGeom->ConvexElems[PrimData.PrimIndex].ScaleElem(ModifiedScale, MinPrimSize);
+			break;
+		}
+	}
+}
+
+bool FStaticMeshEditor::CalcSelectedPrimsAABB(FBox &OutBox) const
+{
+	check(StaticMesh->BodySetup);
+
+	FKAggregateGeom* AggGeom = &StaticMesh->BodySetup->AggGeom;
+
+	for (int32 PrimIdx = 0; PrimIdx < SelectedPrims.Num(); PrimIdx++)
+	{
+		const FPrimData& PrimData = SelectedPrims[PrimIdx];
+
+		check(IsPrimValid(PrimData));
+		switch (PrimData.PrimType)
+		{
+		case KPT_Sphere:
+			OutBox += AggGeom->SphereElems[PrimData.PrimIndex].CalcAABB(FTransform::Identity, 1.f);
+			break;
+		case KPT_Box:
+			OutBox += AggGeom->BoxElems[PrimData.PrimIndex].CalcAABB(FTransform::Identity, 1.f);
+			break;
+		case KPT_Sphyl:
+			OutBox += AggGeom->SphylElems[PrimData.PrimIndex].CalcAABB(FTransform::Identity, 1.f);
+			break;
+		case KPT_Convex:
+			OutBox += AggGeom->ConvexElems[PrimData.PrimIndex].CalcAABB(FTransform::Identity, FVector(1.f));
+			break;
+		}
+	}
+	return HasSelectedPrims();
+}
+
+bool FStaticMeshEditor::GetLastSelectedPrimTransform(FTransform& OutTransform) const
+{
+	if (SelectedPrims.Num() > 0)
+	{
+		check(StaticMesh->BodySetup);
+
+		const FKAggregateGeom* AggGeom = &StaticMesh->BodySetup->AggGeom;
+
+		const FPrimData& PrimData = SelectedPrims.Last();
+
+		check(IsPrimValid(PrimData));
+		switch (PrimData.PrimType)
+		{
+		case KPT_Sphere:
+			OutTransform = AggGeom->SphereElems[PrimData.PrimIndex].GetTransform();
+			break;
+		case KPT_Box:
+			OutTransform = AggGeom->BoxElems[PrimData.PrimIndex].GetTransform();
+			break;
+		case KPT_Sphyl:
+			OutTransform = AggGeom->SphylElems[PrimData.PrimIndex].GetTransform();
+			break;
+		case KPT_Convex:
+			OutTransform = AggGeom->ConvexElems[PrimData.PrimIndex].GetTransform();
+			break;
+		}
+	}
+	return HasSelectedPrims();
+}
+
+FTransform FStaticMeshEditor::GetPrimTransform(const FPrimData& InPrimData) const
+{
+	check(StaticMesh->BodySetup);
+
+	const FKAggregateGeom* AggGeom = &StaticMesh->BodySetup->AggGeom;
+
+	check(IsPrimValid(InPrimData));
+	switch (InPrimData.PrimType)
+	{
+	case KPT_Sphere:
+		return AggGeom->SphereElems[InPrimData.PrimIndex].GetTransform();
+	case KPT_Box:
+		return AggGeom->BoxElems[InPrimData.PrimIndex].GetTransform();
+	case KPT_Sphyl:
+		return AggGeom->SphylElems[InPrimData.PrimIndex].GetTransform();
+	case KPT_Convex:
+		return AggGeom->ConvexElems[InPrimData.PrimIndex].GetTransform();
+	}
+	return FTransform::Identity;
+}
+
+void FStaticMeshEditor::SetPrimTransform(const FPrimData& InPrimData, const FTransform& InPrimTransform) const
+{
+	check(StaticMesh->BodySetup);
+
+	FKAggregateGeom* AggGeom = &StaticMesh->BodySetup->AggGeom;
+
+	check(IsPrimValid(InPrimData));
+	switch (InPrimData.PrimType)
+	{
+	case KPT_Sphere:
+		AggGeom->SphereElems[InPrimData.PrimIndex].SetTransform(InPrimTransform);
+		break;
+	case KPT_Box:
+		AggGeom->BoxElems[InPrimData.PrimIndex].SetTransform(InPrimTransform);
+		break;
+	case KPT_Sphyl:
+		AggGeom->SphylElems[InPrimData.PrimIndex].SetTransform(InPrimTransform);
+		break;
+	case KPT_Convex:
+		AggGeom->ConvexElems[InPrimData.PrimIndex].SetTransform(InPrimTransform);
+		break;
+	}
+}
+
 void FStaticMeshEditor::RefreshTool()
 {
 	int32 NumLODs = StaticMesh->GetNumLODs();
@@ -729,19 +1052,61 @@ int32 FStaticMeshEditor::GetCurrentLODIndex()
 void FStaticMeshEditor::GenerateKDop(const FVector* Directions, uint32 NumDirections)
 {
 	TArray<FVector>	DirArray;
-
 	for(uint32 DirectionIndex = 0;DirectionIndex < NumDirections;DirectionIndex++)
 	{
 		DirArray.Add(Directions[DirectionIndex]);
 	}
-	GenerateKDopAsSimpleCollision( StaticMesh, DirArray );
+
+	GEditor->BeginTransaction(LOCTEXT("FStaticMeshEditor_GenerateKDop", "Create Convex Collision"));
+	const int32 PrimIndex = GenerateKDopAsSimpleCollision(StaticMesh, DirArray);
+	GEditor->EndTransaction();
+	if (PrimIndex != INDEX_NONE)
+	{
+		ClearSelectedPrims();
+		AddSelectedPrim(FPrimData(KPT_Convex, PrimIndex));
+	}
+
+	Viewport->RefreshViewport();
+}
+
+void FStaticMeshEditor::OnCollisionBox()
+{
+	GEditor->BeginTransaction(LOCTEXT("FStaticMeshEditor_OnCollisionBox", "Create Box Collision"));
+	const int32 PrimIndex = GenerateBoxAsSimpleCollision(StaticMesh);
+	GEditor->EndTransaction();
+	if (PrimIndex != INDEX_NONE)
+	{
+		ClearSelectedPrims();
+		AddSelectedPrim(FPrimData(KPT_Box, PrimIndex));
+	}
 
 	Viewport->RefreshViewport();
 }
 
 void FStaticMeshEditor::OnCollisionSphere()
 {
-	GenerateSphereAsSimpleCollision( StaticMesh );
+	GEditor->BeginTransaction(LOCTEXT("FStaticMeshEditor_OnCollisionSphere", "Create Sphere Collision"));
+	const int32 PrimIndex = GenerateSphereAsSimpleCollision(StaticMesh);
+	GEditor->EndTransaction();
+	if (PrimIndex != INDEX_NONE)
+	{
+		ClearSelectedPrims();
+		AddSelectedPrim(FPrimData(KPT_Sphere, PrimIndex));
+	}
+
+	Viewport->RefreshViewport();
+}
+
+void FStaticMeshEditor::OnCollisionSphyl()
+{
+	GEditor->BeginTransaction(LOCTEXT("FStaticMeshEditor_OnCollisionSphyl", "Create Capsule Collision"));
+	const int32 PrimIndex = GenerateSphylAsSimpleCollision(StaticMesh);
+	GEditor->EndTransaction();
+	if (PrimIndex != INDEX_NONE)
+	{
+		ClearSelectedPrims();
+		AddSelectedPrim(FPrimData(KPT_Sphyl, PrimIndex));
+	}
 
 	Viewport->RefreshViewport();
 }
@@ -752,9 +1117,11 @@ void FStaticMeshEditor::OnRemoveCollision(void)
 	UBodySetup* BS = StaticMesh->BodySetup;
 	if (BS != NULL && (BS->AggGeom.GetElementCount() > 0))
 	{
-		int32 ShouldReplace = FMessageDialog::Open(EAppMsgType::YesNo, LOCTEXT("RemoveCollisionPrompt", "Are you sure you want to remove the collision mesh?"));
+		int32 ShouldReplace = FMessageDialog::Open(EAppMsgType::YesNo, LOCTEXT("RemoveCollisionPrompt", "Are you sure you want to remove all the collision meshes?"));
 		if (ShouldReplace == EAppReturnType::Yes)
 		{
+			ClearSelectedPrims();
+
 			// Make sure rendering is done - so we are not changing data being used by collision drawing.
 			FlushRenderingCommands();
 
@@ -845,6 +1212,8 @@ void FStaticMeshEditor::OnConvertBoxToConvexCollision()
 			int32 NumBoxElems = BodySetup->AggGeom.BoxElems.Num();
 			if (NumBoxElems > 0)
 			{
+				ClearSelectedPrims();
+
 				// Make sure rendering is done - so we are not changing data being used by collision drawing.
 				FlushRenderingCommands();
 
@@ -918,6 +1287,8 @@ void FStaticMeshEditor::OnCopyCollisionFromSelectedStaticMesh()
 		{
 			UBodySetup* BodySetup = StaticMesh->BodySetup;
 
+			ClearSelectedPrims();
+
 			// Make sure rendering is done - so we are not changing data being used by collision drawing.
 			FlushRenderingCommands();
 
@@ -1055,6 +1426,8 @@ void FStaticMeshEditor::DoDecomp(int32 InMaxHullCount, int32 InMaxHullVerts)
 		TArray<uint32> Indices;
 		LODModel.IndexBuffer.GetCopy(Indices);
 
+		ClearSelectedPrims();
+
 		// Make sure rendering is done - so we are not changing data being used by collision drawing.
 		FlushRenderingCommands();
 
@@ -1146,6 +1519,24 @@ int32 FStaticMeshEditor::GetNumUVChannels( int32 LODLevel ) const
 	return NumUVChannels.IsValidIndex(LODLevel) ? NumUVChannels[LODLevel] : 0;
 }
 
+void FStaticMeshEditor::DeleteSelected()
+{
+	if (GetSelectedSocket())
+	{
+		DeleteSelectedSockets();
+	}
+	
+	if (HasSelectedPrims())
+	{
+		DeleteSelectedPrims();
+	}
+}
+
+bool FStaticMeshEditor::CanDeleteSelected() const
+{
+	return (GetSelectedSocket() != NULL || HasSelectedPrims());
+}
+
 void FStaticMeshEditor::DeleteSelectedSockets()
 {
 	check(SocketManager.IsValid());
@@ -1153,9 +1544,87 @@ void FStaticMeshEditor::DeleteSelectedSockets()
 	SocketManager->DeleteSelectedSocket();
 }
 
-bool FStaticMeshEditor::HasSelectedSockets() const
+void FStaticMeshEditor::DeleteSelectedPrims()
+{
+	if (SelectedPrims.Num() > 0)
+	{
+		//int32 ShouldReplace = FMessageDialog::Open(EAppMsgType::YesNo, LOCTEXT("RemoveSelectedCollisionPrompt", "Are you sure you want to remove the selected collision meshes?"));
+		//if (ShouldReplace == EAppReturnType::Yes)
+		{
+			// Sort the selected prims by PrimIndex so when we're deleting them we don't mess up other prims indicies
+			struct FCompareFPrimDataPrimIndex
+			{
+				FORCEINLINE bool operator()(const FPrimData& A, const FPrimData& B) const
+				{
+					return A.PrimIndex < B.PrimIndex;
+				}
+			};
+			SelectedPrims.Sort(FCompareFPrimDataPrimIndex());
+
+			check(StaticMesh->BodySetup);
+
+			FKAggregateGeom* AggGeom = &StaticMesh->BodySetup->AggGeom;
+
+			GEditor->BeginTransaction(LOCTEXT("FStaticMeshEditor_DeleteSelectedPrims", "Delete Collision"));
+			StaticMesh->BodySetup->Modify();
+
+			for (int32 PrimIdx = SelectedPrims.Num() - 1; PrimIdx >= 0; PrimIdx--)
+			{
+				const FPrimData& PrimData = SelectedPrims[PrimIdx];
+
+				check(IsPrimValid(PrimData));
+				switch (PrimData.PrimType)
+				{
+				case KPT_Sphere:
+					AggGeom->SphereElems.RemoveAt(PrimData.PrimIndex);
+					break;
+				case KPT_Box:
+					AggGeom->BoxElems.RemoveAt(PrimData.PrimIndex);
+					break;
+				case KPT_Sphyl:
+					AggGeom->SphylElems.RemoveAt(PrimData.PrimIndex);
+					break;
+				case KPT_Convex:
+					AggGeom->ConvexElems.RemoveAt(PrimData.PrimIndex);
+					break;
+				}
+			}
+
+			GEditor->EndTransaction();
+
+			ClearSelectedPrims();
+
+			// Make sure rendering is done - so we are not changing data being used by collision drawing.
+			FlushRenderingCommands();
+
+			// refresh collision change back to staticmesh components
+			RefreshCollisionChange(StaticMesh);
+
+			// Mark staticmesh as dirty, to help make sure it gets saved.
+			StaticMesh->MarkPackageDirty();
+
+			// Update views/property windows
+			Viewport->RefreshViewport();
+		}
+	}
+}
+
+void FStaticMeshEditor::DuplicateSelected()
+{
+	DuplicateSelectedSocket();
+
+	const FVector InitialOffset(20.f);
+	DuplicateSelectedPrims(&InitialOffset);
+}
+
+bool FStaticMeshEditor::CanDuplicateSelected() const
 {
-	return ( GetSelectedSocket() != NULL );
+	return (GetSelectedSocket() != NULL || HasSelectedPrims());
+}
+
+bool FStaticMeshEditor::CanRenameSelected() const
+{
+	return (GetSelectedSocket() != NULL);
 }
 
 void FStaticMeshEditor::ExecuteFindInExplorer()
@@ -1257,17 +1726,25 @@ void FStaticMeshEditor::RedoAction()
 
 void FStaticMeshEditor::PostUndo( bool bSuccess )
 {
+	RemoveInvalidPrims();
+
 	OnPostUndo.Broadcast();
 }
 
 void FStaticMeshEditor::PostRedo( bool bSuccess )
 {
+	RemoveInvalidPrims();
+
 	OnPostUndo.Broadcast();
 }
 
 void FStaticMeshEditor::OnSocketSelectionChanged()
 {
 	UStaticMeshSocket* SelectedSocket = GetSelectedSocket();
+	if (SelectedSocket)
+	{
+		ClearSelectedPrims();
+	}
 	Viewport->GetViewportClient().OnSocketSelectionChanged( SelectedSocket );
 }
 
diff --git a/Engine/Source/Editor/StaticMeshEditor/Private/StaticMeshEditor.h b/Engine/Source/Editor/StaticMeshEditor/Private/StaticMeshEditor.h
--- a/Engine/Source/Editor/StaticMeshEditor/Private/StaticMeshEditor.h
+++ b/Engine/Source/Editor/StaticMeshEditor/Private/StaticMeshEditor.h
@@ -12,6 +12,8 @@ class FStaticMeshEditor : public IStaticMeshEditor, public FGCObject, public FEd
 public:
 	FStaticMeshEditor()
 		: StaticMesh( NULL )
+		, NumLODLevels(0)
+		, MinPrimSize(0.5f)
 	{}
 
 	~FStaticMeshEditor();
@@ -58,6 +60,22 @@ public:
 	virtual void DuplicateSelectedSocket() OVERRIDE;
 	virtual void RequestRenameSelectedSocket() OVERRIDE;
 
+	virtual bool IsPrimValid(const FPrimData& InPrimData) const override;
+	virtual bool HasSelectedPrims() const;
+	virtual void AddSelectedPrim(const FPrimData& InPrimData) override;
+	virtual void RemoveSelectedPrim(const FPrimData& InPrimData) override;
+	virtual void RemoveInvalidPrims() override;
+	virtual bool IsSelectedPrim(const FPrimData& InPrimData) const override;
+	virtual void ClearSelectedPrims() override;
+	virtual void DuplicateSelectedPrims(const FVector* InOffset) override;
+	virtual void TranslateSelectedPrims(const FVector& InDrag) override;
+	virtual void RotateSelectedPrims(const FRotator& InRot) override;
+	virtual void ScaleSelectedPrims(const FVector& InScale) override;
+	virtual bool CalcSelectedPrimsAABB(FBox &OutBox) const override;
+	virtual bool GetLastSelectedPrimTransform(FTransform& OutTransform) const;
+	FTransform GetPrimTransform(const FPrimData& InPrimData) const;
+	void SetPrimTransform(const FPrimData& InPrimData, const FTransform& InPrimTransform) const;
+
 	virtual int32 GetNumTriangles( int32 LODLevel = 0 ) const OVERRIDE;
 	virtual int32 GetNumVertices( int32 LODLevel = 0 ) const OVERRIDE;
 	virtual int32 GetNumUVChannels( int32 LODLevel = 0 ) const OVERRIDE;
@@ -129,9 +147,15 @@ private:
 	/** Helper function for generating K-DOP collision geometry. */
 	void GenerateKDop(const FVector* Directions, uint32 NumDirections);
 
+	/** Callback for creating box collision. */
+	void OnCollisionBox();
+
 	/** Callback for creating sphere collision. */
 	void OnCollisionSphere();
 
+	/** Callback for creating sphyl collision. */
+	void OnCollisionSphyl();
+
 	/** Event for exporting the light map channel of a static mesh to an intermediate file for Max/Maya */
 	void OnExportLightmapMesh( bool IsFBX );
 
@@ -165,11 +189,26 @@ private:
 	/** Rebuilds the UV Channel combo list and attempts to set it to the same channel. */
 	void RegenerateUVChannelComboList();
 
+	/** Delete whats currently selected */
+	void DeleteSelected();
+
+	/** Whether we currently have any selected that can be deleted */
+	bool CanDeleteSelected() const;
+
 	/** Delete the currently selected sockets */
 	void DeleteSelectedSockets();
 
-	/** Whether we currently have selected sockets */
-	bool HasSelectedSockets() const;
+	/** Delete the currently selected prims */
+	void DeleteSelectedPrims();
+
+	/** Duplicate whats currently selected */
+	void DuplicateSelected();
+
+	/** Whether we currently have any selected that can be duplicated */
+	bool CanDuplicateSelected() const;
+
+	/** Whether we currently have any selected that can be renamed */
+	bool CanRenameSelected() const;
 
 	/** Handler for when FindInExplorer is selected */
 	void ExecuteFindInExplorer();
@@ -256,6 +295,12 @@ private:
 	/** Delegates called after an undo operation for child widgets to refresh */
 	FOnPostUndoMulticaster OnPostUndo;	
 
+	/** Information on the selected collision primitives */
+	TArray<FPrimData> SelectedPrims;
+
+	/** Misc consts */
+	const float	MinPrimSize;
+
 	/**	The tab ids for all the tabs used */
 	static const FName ViewportTabId;
 	static const FName PropertiesTabId;
diff --git a/Engine/Source/Editor/StaticMeshEditor/Private/StaticMeshEditorActions.cpp b/Engine/Source/Editor/StaticMeshEditor/Private/StaticMeshEditorActions.cpp
--- a/Engine/Source/Editor/StaticMeshEditor/Private/StaticMeshEditorActions.cpp
+++ b/Engine/Source/Editor/StaticMeshEditor/Private/StaticMeshEditorActions.cpp
@@ -22,14 +22,15 @@ void FStaticMeshEditorCommands::RegisterCommands()
 	UI_COMMAND( SetShowBinormals, "Binormals", "Toggles display of vertex binormals (orthogonal vector to normal and tangent) in the Preview Pane.", EUserInterfaceActionType::ToggleButton, FInputGesture() );
 	UI_COMMAND( SetShowPivot, "Show Pivot", "Display the pivot location of the static mesh.", EUserInterfaceActionType::ToggleButton, FInputGesture() );	
 
-	UI_COMMAND( CreateDOP6, "6DOP Simplified Collision", "Generates a new axis-aligned box collision mesh (6 total sides) encompassing the static mesh.", EUserInterfaceActionType::Button, FInputGesture() );
 	UI_COMMAND( CreateDOP10X, "10DOP-X Simplified Collision", "Generates a new axis-aligned box collision mesh with the 4 X-axis aligned edges beveled (10 total sides) encompassing the static mesh.", EUserInterfaceActionType::Button, FInputGesture() );
 	UI_COMMAND( CreateDOP10Y, "10DOP-Y Simplified Collision", "Generates a new axis-aligned box collision mesh with the 4 Y-axis aligned edges beveled (10 total sides) encompassing the static mesh.", EUserInterfaceActionType::Button, FInputGesture() );
 	UI_COMMAND( CreateDOP10Z, "10DOP-Z Simplified Collision", "Generates a new axis-aligned box collision mesh with the 4 Z-axis aligned edges beveled (10 total sides) encompassing the static mesh.", EUserInterfaceActionType::Button, FInputGesture() );
 	UI_COMMAND( CreateDOP18, "18DOP Simplified Collision", "Generates a new axis-aligned box collision mesh with all edges beveled (18 total sides) encompassing the static mesh.", EUserInterfaceActionType::Button, FInputGesture() );
 	UI_COMMAND( CreateDOP26, "26DOP Simplified Collision", "Generates a new axis-aligned box collision mesh with all edges and corners beveled (26 total sides) encompassing the static mesh.", EUserInterfaceActionType::Button, FInputGesture() );
 
+	UI_COMMAND( CreateBoxCollision, "Box Simplified Collision", "Generates a new box collision mesh encompassing the static mesh.", EUserInterfaceActionType::Button, FInputGesture() );
 	UI_COMMAND( CreateSphereCollision, "Sphere Simplified Collision", "Generates a new sphere collision mesh encompassing the static mesh.", EUserInterfaceActionType::Button, FInputGesture() );
+	UI_COMMAND( CreateSphylCollision, "Capsule Simplified Collision", "Generates a new capsule collision mesh encompassing the static mesh.", EUserInterfaceActionType::Button, FInputGesture() );
 	UI_COMMAND( CreateAutoConvexCollision, "Auto Convex Collision", "Opens the Auto Convex Collision Tool for generating a new convex collision mesh, or meshes.", EUserInterfaceActionType::Button, FInputGesture() );
 	UI_COMMAND( RemoveCollision, "Remove Collision", "Removes any simplified collision assigned to the static mesh.", EUserInterfaceActionType::Button, FInputGesture() );
 	UI_COMMAND( ConvertBoxesToConvex, "Convert Boxes to Convex", "Converts any simple box collision meshes to convex collision meshes.", EUserInterfaceActionType::Button, FInputGesture() );
diff --git a/Engine/Source/Editor/StaticMeshEditor/Private/StaticMeshEditorViewportClient.cpp b/Engine/Source/Editor/StaticMeshEditor/Private/StaticMeshEditorViewportClient.cpp
--- a/Engine/Source/Editor/StaticMeshEditor/Private/StaticMeshEditorViewportClient.cpp
+++ b/Engine/Source/Editor/StaticMeshEditor/Private/StaticMeshEditorViewportClient.cpp
@@ -10,7 +10,6 @@
 #include "StaticMeshResources.h"
 #include "RawMesh.h"
 
-
 #include "StaticMeshEditor.h"
 #include "BusyCursor.h"
 #include "MeshBuild.h"
@@ -63,6 +62,7 @@ FStaticMeshEditorViewportClient::FStaticMeshEditorViewportClient(TWeakPtr<IStati
 	OverrideNearClipPlane(1.0f);
 	bUsingOrbitCamera = true;
 
+	bShowCollision = false;
 	bShowSockets = true;
 	bDrawUVs = false;
 	bDrawNormals = false;
@@ -87,6 +87,25 @@ void FStaticMeshEditorViewportClient::Tick(float DeltaSeconds)
 	}
 }
 
+/**
+ * A hit proxy class for the wireframe collision geometry
+ */
+struct HSMECollisionProxy : public HHitProxy
+{
+	DECLARE_HIT_PROXY();
+
+	IStaticMeshEditor::FPrimData	PrimData;
+
+	HSMECollisionProxy(const IStaticMeshEditor::FPrimData& InPrimData) :
+		HHitProxy(HPP_UI),
+		PrimData(InPrimData) {}
+
+	HSMECollisionProxy(EKCollisionPrimitiveType InPrimType, int32 InPrimIndex) :
+		HHitProxy(HPP_UI),
+		PrimData(InPrimType, InPrimIndex) {}
+};
+IMPLEMENT_HIT_PROXY(HSMECollisionProxy, HHitProxy);
+
 /**
  * A hit proxy class for sockets.
  */
@@ -94,12 +113,11 @@ struct HSMESocketProxy : public HHitProxy
 {
 	DECLARE_HIT_PROXY( );
 
-	int32		SocketIndex;
+	int32							SocketIndex;
 
-	HSMESocketProxy(int32 InSocketIndex)
-		:	HHitProxy( HPP_UI )
-		,	SocketIndex( InSocketIndex )
-	{}
+	HSMESocketProxy(int32 InSocketIndex) :
+		HHitProxy( HPP_UI ), 
+		SocketIndex( InSocketIndex ) {}
 };
 IMPLEMENT_HIT_PROXY(HSMESocketProxy, HHitProxy);
 
@@ -108,51 +126,74 @@ bool FStaticMeshEditorViewportClient::InputWidgetDelta( FViewport* Viewport, EAx
 	bool bHandled = false;
 	if (bManipulating)
 	{
-		UStaticMeshSocket* SelectedSocket = StaticMeshEditorPtr.Pin()->GetSelectedSocket();
-
-		if(SelectedSocket && CurrentAxis != EAxisList::None)
+		if (CurrentAxis != EAxisList::None)
 		{
-			UProperty* ChangedProperty = NULL;
-			const FWidget::EWidgetMode MoveMode = GetWidgetMode();
-			if(MoveMode == FWidget::WM_Rotate)
-			{
-				ChangedProperty = FindField<UProperty>( UStaticMeshSocket::StaticClass(), "RelativeRotation" );
-				SelectedSocket->PreEditChange(ChangedProperty);
-
-				FRotator CurrentRot = SelectedSocket->RelativeRotation;
-				FRotator SocketWinding, SocketRotRemainder;
-				CurrentRot.GetWindingAndRemainder(SocketWinding, SocketRotRemainder);
-
-				const FQuat ActorQ = SocketRotRemainder.Quaternion();
-				const FQuat DeltaQ = Rot.Quaternion();
-				const FQuat ResultQ = DeltaQ * ActorQ;
-				const FRotator NewSocketRotRem = FRotator( ResultQ );
-				FRotator DeltaRot = NewSocketRotRem - SocketRotRemainder;
-				DeltaRot.Normalize();
-
-				SelectedSocket->RelativeRotation += DeltaRot;
-				SelectedSocket->RelativeRotation.Pitch = FRotator::ClampAxis( SelectedSocket->RelativeRotation.Pitch );
-				SelectedSocket->RelativeRotation.Yaw = FRotator::ClampAxis( SelectedSocket->RelativeRotation.Yaw );
-				SelectedSocket->RelativeRotation.Roll = FRotator::ClampAxis( SelectedSocket->RelativeRotation.Roll );
-			}
-			else if(MoveMode == FWidget::WM_Translate)
+			UStaticMeshSocket* SelectedSocket = StaticMeshEditorPtr.Pin()->GetSelectedSocket();
+			if(SelectedSocket)
 			{
-				ChangedProperty = FindField<UProperty>( UStaticMeshSocket::StaticClass(), "RelativeLocation" );
-				SelectedSocket->PreEditChange(ChangedProperty);
+				UProperty* ChangedProperty = NULL;
+				const FWidget::EWidgetMode MoveMode = GetWidgetMode();
+				if(MoveMode == FWidget::WM_Rotate)
+				{
+					ChangedProperty = FindField<UProperty>( UStaticMeshSocket::StaticClass(), "RelativeRotation" );
+					SelectedSocket->PreEditChange(ChangedProperty);
+
+					FRotator CurrentRot = SelectedSocket->RelativeRotation;
+					FRotator SocketWinding, SocketRotRemainder;
+					CurrentRot.GetWindingAndRemainder(SocketWinding, SocketRotRemainder);
+
+					const FQuat ActorQ = SocketRotRemainder.Quaternion();
+					const FQuat DeltaQ = Rot.Quaternion();
+					const FQuat ResultQ = DeltaQ * ActorQ;
+					const FRotator NewSocketRotRem = FRotator( ResultQ );
+					FRotator DeltaRot = NewSocketRotRem - SocketRotRemainder;
+					DeltaRot.Normalize();
+
+					SelectedSocket->RelativeRotation += DeltaRot;
+					SelectedSocket->RelativeRotation = SelectedSocket->RelativeRotation.Clamp();
+				}
+				else if(MoveMode == FWidget::WM_Translate)
+				{
+					ChangedProperty = FindField<UProperty>( UStaticMeshSocket::StaticClass(), "RelativeLocation" );
+					SelectedSocket->PreEditChange(ChangedProperty);
 
-				//FRotationMatrix SocketRotTM( SelectedSocket->RelativeRotation );
-				//FVector SocketMove = SocketRotTM.TransformVector( Drag );
+					//FRotationMatrix SocketRotTM( SelectedSocket->RelativeRotation );
+					//FVector SocketMove = SocketRotTM.TransformVector( Drag );
 
-				SelectedSocket->RelativeLocation += Drag;
-			}
-			if ( ChangedProperty )
-			{			
-				FPropertyChangedEvent PropertyChangedEvent( ChangedProperty );
-				SelectedSocket->PostEditChangeProperty(PropertyChangedEvent);
+					SelectedSocket->RelativeLocation += Drag;
+				}
+				if ( ChangedProperty )
+				{			
+					FPropertyChangedEvent PropertyChangedEvent( ChangedProperty );
+					SelectedSocket->PostEditChangeProperty(PropertyChangedEvent);
+				}
+
+				StaticMeshEditorPtr.Pin()->GetStaticMesh()->MarkPackageDirty();
 			}
+			else
+			{
+				const bool bSelectedPrim = StaticMeshEditorPtr.Pin()->HasSelectedPrims();
+				if (bSelectedPrim && CurrentAxis != EAxisList::None)
+				{
+					const FWidget::EWidgetMode MoveMode = GetWidgetMode();
+					if (MoveMode == FWidget::WM_Rotate)
+					{
+						StaticMeshEditorPtr.Pin()->RotateSelectedPrims(Rot);
+					}
+					else if (MoveMode == FWidget::WM_Scale)
+					{
+						StaticMeshEditorPtr.Pin()->ScaleSelectedPrims(Scale);
+					}
+					else if (MoveMode == FWidget::WM_Translate)
+					{
+						StaticMeshEditorPtr.Pin()->TranslateSelectedPrims(Drag);
+					}
 
-			StaticMeshEditorPtr.Pin()->GetStaticMesh()->MarkPackageDirty();
+					StaticMeshEditorPtr.Pin()->GetStaticMesh()->MarkPackageDirty();
+				}
+			}
 		}
+
 		Invalidate();		
 		bHandled = true;
 	}
@@ -164,25 +205,71 @@ void FStaticMeshEditorViewportClient::TrackingStarted( const struct FInputEventS
 {
 	if( !bManipulating && bIsDraggingWidget )
 	{
-		if( GetWidgetMode() == FWidget::WM_Rotate )
+		const UStaticMeshSocket* SelectedSocket = StaticMeshEditorPtr.Pin()->GetSelectedSocket();
+		if (SelectedSocket)
 		{
-			GEditor->BeginTransaction( LOCTEXT("FStaticMeshEditorViewportClient_RotateSocket", "Rotate Socket" ) );
+			FText TransText;
+			if( GetWidgetMode() == FWidget::WM_Rotate )
+			{
+				TransText = LOCTEXT("FStaticMeshEditorViewportClient_RotateSocket", "Rotate Socket");
+			}
+			else if (GetWidgetMode() == FWidget::WM_Translate)
+			{
+				if( InInputState.IsLeftMouseButtonPressed() && (Widget->GetCurrentAxis() & EAxisList::XYZ) )
+				{
+					const bool bAltDown = InInputState.IsAltButtonPressed();
+					if ( bAltDown )
+					{
+						// Rather than moving/rotating the selected socket, copy it and move the copy instead
+						StaticMeshEditorPtr.Pin()->DuplicateSelectedSocket();
+					}
+				}
+
+				TransText = LOCTEXT("FStaticMeshEditorViewportClient_TranslateSocket", "Translate Socket");
+			}
+
+			if (!TransText.IsEmpty())
+			{
+				GEditor->BeginTransaction(TransText);
+			}
 		}
-		else
+		
+		const bool bSelectedPrim = StaticMeshEditorPtr.Pin()->HasSelectedPrims();
+		if (bSelectedPrim)
 		{
-			if( InInputState.IsLeftMouseButtonPressed() && (Widget->GetCurrentAxis() & EAxisList::XYZ) )
+			FText TransText;
+			if (GetWidgetMode() == FWidget::WM_Rotate)
+			{
+				TransText = LOCTEXT("FStaticMeshEditorViewportClient_RotateCollision", "Rotate Collision");
+			}
+			else if (GetWidgetMode() == FWidget::WM_Scale)
+			{
+				TransText = LOCTEXT("FStaticMeshEditorViewportClient_ScaleCollision", "Scale Collision");
+			}
+			else if (GetWidgetMode() == FWidget::WM_Translate)
+			{
+				if (InInputState.IsLeftMouseButtonPressed() && (Widget->GetCurrentAxis() & EAxisList::XYZ))
+				{
+					const bool bAltDown = InInputState.IsAltButtonPressed();
+					if (bAltDown)
+					{
+						// Rather than moving/rotating the selected primitives, copy them and move the copies instead
+						StaticMeshEditorPtr.Pin()->DuplicateSelectedPrims(NULL);
+					}
+				}
+
+				TransText = LOCTEXT("FStaticMeshEditorViewportClient_TranslateCollision", "Translate Collision");
+			}
+			if (!TransText.IsEmpty())
 			{
-				const bool bAltDown = InInputState.IsAltButtonPressed();
-				if ( bAltDown )
+				GEditor->BeginTransaction(TransText);
+				if (StaticMesh->BodySetup)
 				{
-					// Rather than moving/rotating the selected socket, copy it and move the copy instead
-					StaticMeshEditorPtr.Pin()->DuplicateSelectedSocket();
+					StaticMesh->BodySetup->Modify();
 				}
 			}
+		}
 
-			GEditor->BeginTransaction( LOCTEXT("FStaticMeshEditorViewportClient_TranslateSocket", "Translate Socket" ) );
-		}	
-		
 		bManipulating = true;
 	}
 
@@ -190,12 +277,18 @@ void FStaticMeshEditorViewportClient::TrackingStarted( const struct FInputEventS
 
 FWidget::EWidgetMode FStaticMeshEditorViewportClient::GetWidgetMode() const
 {
-	UStaticMeshSocket* SelectedSocket = StaticMeshEditorPtr.Pin()->GetSelectedSocket();
+	const UStaticMeshSocket* SelectedSocket = StaticMeshEditorPtr.Pin()->GetSelectedSocket();
 	if( SelectedSocket )
 	{
 		return WidgetMode;
 	}
 
+	const bool bSelectedPrim = StaticMeshEditorPtr.Pin()->HasSelectedPrims();
+	if (bSelectedPrim)
+	{
+		return WidgetMode;
+	}
+
 	return FWidget::WM_None;
 }
 
@@ -210,7 +303,7 @@ void FStaticMeshEditorViewportClient::TrackingStopped()
 
 FVector FStaticMeshEditorViewportClient::GetWidgetLocation() const
 {
-	UStaticMeshSocket* SelectedSocket = StaticMeshEditorPtr.Pin()->GetSelectedSocket();
+	const UStaticMeshSocket* SelectedSocket = StaticMeshEditorPtr.Pin()->GetSelectedSocket();
 	if( SelectedSocket )
 	{
 		FMatrix SocketTM;
@@ -219,12 +312,19 @@ FVector FStaticMeshEditorViewportClient::GetWidgetLocation() const
 		return SocketTM.GetOrigin();
 	}
 
+	FTransform PrimTransform = FTransform::Identity;
+	const bool bSelectedPrim = StaticMeshEditorPtr.Pin()->GetLastSelectedPrimTransform(PrimTransform);
+	if (bSelectedPrim)
+	{
+		return PrimTransform.GetLocation();
+	}
+
 	return FVector::ZeroVector;
 }
 
 FMatrix FStaticMeshEditorViewportClient::GetWidgetCoordSystem() const 
 {
-	UStaticMeshSocket* SelectedSocket = StaticMeshEditorPtr.Pin()->GetSelectedSocket();
+	const UStaticMeshSocket* SelectedSocket = StaticMeshEditorPtr.Pin()->GetSelectedSocket();
 	if( SelectedSocket )
 	{
 		//FMatrix SocketTM;
@@ -233,6 +333,13 @@ FMatrix FStaticMeshEditorViewportClient::GetWidgetCoordSystem() const
 		return FRotationMatrix( SelectedSocket->RelativeRotation );
 	}
 
+	FTransform PrimTransform = FTransform::Identity;
+	const bool bSelectedPrim = StaticMeshEditorPtr.Pin()->GetLastSelectedPrimTransform(PrimTransform);
+	if (bSelectedPrim)
+	{
+		return FRotationMatrix(PrimTransform.Rotator());
+	}
+
 	return FMatrix::Identity;
 }
 
@@ -251,8 +358,71 @@ void FStaticMeshEditorViewportClient::Draw(const FSceneView* View,FPrimitiveDraw
 {
 	FEditorViewportClient::Draw(View, PDI);
 
+	if (bShowCollision && StaticMesh->BodySetup)
+	{
+		const FColor SelectedColor(149, 223, 157);
+		const FColor UnselectedColor(157, 149, 223);
+
+		// Draw bodies
+		FKAggregateGeom* AggGeom = &StaticMesh->BodySetup->AggGeom;
+
+		for (int32 i = 0; i < AggGeom->SphereElems.Num(); ++i)
+		{
+			HSMECollisionProxy* HitProxy = new HSMECollisionProxy(KPT_Sphere, i);
+			PDI->SetHitProxy(HitProxy);
+
+			const FColor CollisionColor = StaticMeshEditorPtr.Pin()->IsSelectedPrim(HitProxy->PrimData) ? SelectedColor : UnselectedColor;
+			const FKSphereElem& SphereElem = AggGeom->SphereElems[i];
+			const FTransform ElemTM = SphereElem.GetTransform();
+			SphereElem.DrawElemWire(PDI, ElemTM, 1.f, CollisionColor);
+
+			PDI->SetHitProxy(NULL);
+		}
+
+		for (int32 i = 0; i < AggGeom->BoxElems.Num(); ++i)
+		{
+			HSMECollisionProxy* HitProxy = new HSMECollisionProxy(KPT_Box, i);
+			PDI->SetHitProxy(HitProxy);
+
+			const FColor CollisionColor = StaticMeshEditorPtr.Pin()->IsSelectedPrim(HitProxy->PrimData) ? SelectedColor : UnselectedColor;
+			const FKBoxElem& BoxElem = AggGeom->BoxElems[i];
+			const FTransform ElemTM = BoxElem.GetTransform();
+			BoxElem.DrawElemWire(PDI, ElemTM, 1.f, CollisionColor);
+
+			PDI->SetHitProxy(NULL);
+		}
+
+		for (int32 i = 0; i < AggGeom->SphylElems.Num(); ++i)
+		{
+			HSMECollisionProxy* HitProxy = new HSMECollisionProxy(KPT_Sphyl, i);
+			PDI->SetHitProxy(HitProxy);
+
+			const FColor CollisionColor = StaticMeshEditorPtr.Pin()->IsSelectedPrim(HitProxy->PrimData) ? SelectedColor : UnselectedColor;
+			const FKSphylElem& SphylElem = AggGeom->SphylElems[i];
+			const FTransform ElemTM = SphylElem.GetTransform();
+			SphylElem.DrawElemWire(PDI, ElemTM, 1.f, CollisionColor);
+
+			PDI->SetHitProxy(NULL);
+		}
+
+		for (int32 i = 0; i < AggGeom->ConvexElems.Num(); ++i)
+		{
+			HSMECollisionProxy* HitProxy = new HSMECollisionProxy(KPT_Convex, i);
+			PDI->SetHitProxy(HitProxy);
+
+			const FColor CollisionColor = StaticMeshEditorPtr.Pin()->IsSelectedPrim(HitProxy->PrimData) ? SelectedColor : UnselectedColor;
+			const FKConvexElem& ConvexElem = AggGeom->ConvexElems[i];
+			const FTransform ElemTM = ConvexElem.GetTransform();
+			ConvexElem.DrawElemWire(PDI, ElemTM, CollisionColor);
+
+			PDI->SetHitProxy(NULL);
+		}
+	}
+
 	if( bShowSockets )
 	{
+		const FColor SocketColor = FColor(255, 128, 128);
+
 		for(int32 i=0; i < StaticMesh->Sockets.Num(); i++)
 		{
 			UStaticMeshSocket* Socket = StaticMesh->Sockets[i];
@@ -261,7 +431,7 @@ void FStaticMeshEditorViewportClient::Draw(const FSceneView* View,FPrimitiveDraw
 				FMatrix SocketTM;
 				Socket->GetSocketMatrix(SocketTM, StaticMeshComponent);
 				PDI->SetHitProxy( new HSMESocketProxy(i) );
-				DrawWireDiamond(PDI, SocketTM, 5.f, FColor(255,128,128), SDPG_Foreground );
+				DrawWireDiamond(PDI, SocketTM, 5.f, SocketColor, SDPG_Foreground);
 				PDI->SetHitProxy( NULL );
 			}
 		}
@@ -436,7 +606,8 @@ void FStaticMeshEditorViewportClient::DrawCanvas( FViewport& InViewport, FSceneV
 					FCanvasTextItem TextItem( FVector2D( XPos, YPos ), FText::FromString( Socket->SocketName.ToString() ), GEngine->GetSmallFont(), FLinearColor(FColor(255,196,196)) );
 					Canvas.DrawItem( TextItem );	
 
-					if (bManipulating && StaticMeshEditorPtr.Pin()->GetSelectedSocket() == Socket)
+					const UStaticMeshSocket* SelectedSocket = StaticMeshEditorPtr.Pin()->GetSelectedSocket();
+					if (bManipulating && SelectedSocket == Socket)
 					{
 						//Figure out the text height
 						FTextSizingParameters Parameters(GEngine->GetSmallFont(), 1.0f, 1.0f);
@@ -562,7 +733,7 @@ void FStaticMeshEditorViewportClient::DrawCanvas( FViewport& InViewport, FSceneV
 	YPos += 18;
 
 	// Show the number of collision primitives if we are drawing collision.
-	if(EngineShowFlags.Collision && StaticMesh->BodySetup)
+	if(bShowCollision && StaticMesh->BodySetup)
 	{
 		Canvas.DrawShadowedString(
 			6,
@@ -607,20 +778,78 @@ bool FStaticMeshEditorViewportClient::InputKey(FViewport* Viewport, int32 Contro
 
 	if( Event == IE_Pressed) 
 	{
-		if(Key == EKeys::SpaceBar)
+		if (Key == EKeys::W)
 		{
-			if( StaticMeshEditorPtr.Pin()->GetSelectedSocket() && !Widget->IsDragging() )
+			if (!Widget->IsDragging())
 			{
-				const FWidget::EWidgetMode MoveMode = GetWidgetMode();
-				if(MoveMode == FWidget::WM_Rotate)
+				const UStaticMeshSocket* SelectedSocket = StaticMeshEditorPtr.Pin()->GetSelectedSocket();
+				const bool bSelectedPrim = StaticMeshEditorPtr.Pin()->HasSelectedPrims();
+				if ((SelectedSocket || bSelectedPrim))
 				{
 					WidgetMode = FWidget::WM_Translate;
+					Invalidate();
 				}
-				else if(MoveMode == FWidget::WM_Translate)
+			}
+			bHandled = true;
+		}
+		else if (Key == EKeys::E)
+		{
+			if (!Widget->IsDragging())
+			{
+				const UStaticMeshSocket* SelectedSocket = StaticMeshEditorPtr.Pin()->GetSelectedSocket();
+				const bool bSelectedPrim = StaticMeshEditorPtr.Pin()->HasSelectedPrims();
+				if ((SelectedSocket || bSelectedPrim))
 				{
 					WidgetMode = FWidget::WM_Rotate;
+					Invalidate();
+				}
+			}
+			bHandled = true;
+		}
+		else if (Key == EKeys::R)
+		{
+			if (!Widget->IsDragging())
+			{
+				const bool bSelectedPrim = StaticMeshEditorPtr.Pin()->HasSelectedPrims();
+				if ((bSelectedPrim))	// Sockets don't support scaling
+				{
+					WidgetMode = FWidget::WM_Scale;
+					Invalidate();
+				}
+			}
+			bHandled = true;
+		}
+		else if(Key == EKeys::SpaceBar)
+		{
+			if(!Widget->IsDragging())
+			{
+				const UStaticMeshSocket* SelectedSocket = StaticMeshEditorPtr.Pin()->GetSelectedSocket();
+				const bool bSelectedPrim = StaticMeshEditorPtr.Pin()->HasSelectedPrims();
+				if ((SelectedSocket || bSelectedPrim))
+				{
+					const FWidget::EWidgetMode MoveMode = GetWidgetMode();
+					if(MoveMode == FWidget::WM_Rotate)
+					{
+						// Sockets don't support scaling
+						if (bSelectedPrim)
+						{
+							WidgetMode = FWidget::WM_Scale;
+						}
+						else
+						{
+							WidgetMode = FWidget::WM_Translate;
+						}
+					}
+					else if(MoveMode == FWidget::WM_Scale)
+					{
+						WidgetMode = FWidget::WM_Translate;
+					}
+					else if(MoveMode == FWidget::WM_Translate)
+					{
+						WidgetMode = FWidget::WM_Rotate;
+					}
+					Invalidate();
 				}
-				Invalidate();
 			}
 			bHandled = true;
 		}
@@ -645,28 +874,68 @@ bool FStaticMeshEditorViewportClient::InputAxis(FViewport* Viewport, int32 Contr
 
 void FStaticMeshEditorViewportClient::ProcessClick(class FSceneView& InView, class HHitProxy* HitProxy, FKey Key, EInputEvent Event, uint32 HitX, uint32 HitY)
 {
-	if( HitProxy && HitProxy->IsA( HSMESocketProxy::StaticGetType() ) )
-	{
-		HSMESocketProxy* SocketProxy = (HSMESocketProxy*)HitProxy;
+	const bool bCtrlDown = Viewport->KeyState(EKeys::LeftControl) || Viewport->KeyState(EKeys::RightControl);
 
-		UStaticMeshSocket* Socket = NULL;
+	bool ClearSelectedSockets = true;
+	bool ClearSelectedPrims = true;
+	bool ClearSelectedEdges = true;
 
-		if(SocketProxy->SocketIndex < StaticMesh->Sockets.Num())
+	if( HitProxy )
+	{
+		if(HitProxy->IsA( HSMESocketProxy::StaticGetType() ) )
 		{
-			Socket = StaticMesh->Sockets[SocketProxy->SocketIndex];
-		}
+			HSMESocketProxy* SocketProxy = (HSMESocketProxy*)HitProxy;
 
-		if(Socket)
-		{
-			StaticMeshEditorPtr.Pin()->SetSelectedSocket(Socket);
+			UStaticMeshSocket* Socket = NULL;
+
+			if(SocketProxy->SocketIndex < StaticMesh->Sockets.Num())
+			{
+				Socket = StaticMesh->Sockets[SocketProxy->SocketIndex];
+			}
+
+			if(Socket)
+			{
+				StaticMeshEditorPtr.Pin()->SetSelectedSocket(Socket);
+			}
+
+			ClearSelectedSockets = false;
 		}
-		
+		else if (HitProxy->IsA(HSMECollisionProxy::StaticGetType()) && StaticMesh->BodySetup)
+		{
+			HSMECollisionProxy* CollisionProxy = (HSMECollisionProxy*)HitProxy;			
+
+			if (StaticMeshEditorPtr.Pin()->IsSelectedPrim(CollisionProxy->PrimData))
+			{
+				if (!bCtrlDown)
+				{
+					StaticMeshEditorPtr.Pin()->ClearSelectedPrims();
+					StaticMeshEditorPtr.Pin()->AddSelectedPrim(CollisionProxy->PrimData);
+				}
+				else
+				{
+					StaticMeshEditorPtr.Pin()->RemoveSelectedPrim(CollisionProxy->PrimData);
+				}
+			}
+			else
+			{
+				if (!bCtrlDown)
+				{
+					StaticMeshEditorPtr.Pin()->ClearSelectedPrims();
+				}
+				StaticMeshEditorPtr.Pin()->AddSelectedPrim(CollisionProxy->PrimData);
+			}
+
+			// Force the widget to translate, if not already set
+			if (WidgetMode == FWidget::WM_None)
+			{
+				WidgetMode = FWidget::WM_Translate;
+			}
 
-		Invalidate();
+			ClearSelectedPrims = false;
+		}
 	}
-	else if(!HitProxy)
+	else
 	{
-		const bool bCtrlDown = Viewport->KeyState(EKeys::LeftControl) || Viewport->KeyState(EKeys::RightControl);
 		const bool bShiftDown = Viewport->KeyState(EKeys::LeftShift) || Viewport->KeyState(EKeys::RightShift);
 
 		if(!bCtrlDown && !bShiftDown)
@@ -865,11 +1134,26 @@ void FStaticMeshEditorViewportClient::ProcessClick(class FSceneView& InView, cla
 							}
 						}
 					}
+
+					ClearSelectedEdges = false;
 				}
 			}
 		}
 	}
 
+	if (ClearSelectedSockets && StaticMeshEditorPtr.Pin()->GetSelectedSocket())
+	{
+		StaticMeshEditorPtr.Pin()->SetSelectedSocket(NULL);
+	}
+	if (ClearSelectedPrims)
+	{
+		StaticMeshEditorPtr.Pin()->ClearSelectedPrims();
+	}
+	if (ClearSelectedEdges)
+	{
+		SelectedEdgeIndices.Empty();
+	}
+
 	Invalidate();
 }
 
@@ -991,6 +1275,18 @@ bool FStaticMeshEditorViewportClient::IsSetShowBinormalsChecked() const
 	return bDrawBinormals;
 }
 
+void FStaticMeshEditorViewportClient::SetShowWireframeCollision()
+{
+	bShowCollision = !bShowCollision;
+	StaticMeshEditorPtr.Pin()->ClearSelectedPrims();
+	Invalidate();
+}
+
+bool FStaticMeshEditorViewportClient::IsSetShowWireframeCollisionChecked() const
+{
+	return bShowCollision;
+}
+
 void FStaticMeshEditorViewportClient::SetShowSockets()
 {
 	bShowSockets = !bShowSockets;
@@ -1030,9 +1326,14 @@ TSet< int32 >& FStaticMeshEditorViewportClient::GetSelectedEdges()
 
 void FStaticMeshEditorViewportClient::OnSocketSelectionChanged( UStaticMeshSocket* SelectedSocket )
 {
-	if( SelectedSocket && WidgetMode == FWidget::WM_None )
+	if (SelectedSocket)
 	{
-		WidgetMode = FWidget::WM_Translate;
+		SelectedEdgeIndices.Empty();
+
+		if (WidgetMode == FWidget::WM_None || WidgetMode == FWidget::WM_Scale)
+		{
+			WidgetMode = FWidget::WM_Translate;
+		}
 	}
 
 	Invalidate();
diff --git a/Engine/Source/Editor/StaticMeshEditor/Private/StaticMeshEditorViewportClient.h b/Engine/Source/Editor/StaticMeshEditor/Private/StaticMeshEditorViewportClient.h
--- a/Engine/Source/Editor/StaticMeshEditor/Private/StaticMeshEditorViewportClient.h
+++ b/Engine/Source/Editor/StaticMeshEditor/Private/StaticMeshEditorViewportClient.h
@@ -82,6 +82,12 @@ public:
 	/** Callback for checking the binormals show flag. */
 	bool IsSetShowBinormalsChecked() const;
 
+	/** Callback for toggling the collision show flag. */
+	void SetShowWireframeCollision();
+
+	/** Callback for checking the collision show flag. */
+	bool IsSetShowWireframeCollisionChecked() const;
+
 	/** Callback for toggling the socket show flag. */
 	void SetShowSockets();
 
@@ -119,6 +125,7 @@ private:
 
 	/** Flags for various options in the editor. */
 	bool bDrawUVs;
+	bool bShowCollision;
 	bool bShowSockets;
 	bool bDrawNormals;
 	bool bDrawTangents;
diff --git a/Engine/Source/Editor/StaticMeshEditor/Public/IStaticMeshEditor.h b/Engine/Source/Editor/StaticMeshEditor/Public/IStaticMeshEditor.h
--- a/Engine/Source/Editor/StaticMeshEditor/Public/IStaticMeshEditor.h
+++ b/Engine/Source/Editor/StaticMeshEditor/Public/IStaticMeshEditor.h
@@ -4,6 +4,7 @@
 
 #include "Toolkits/IToolkitHost.h"
 #include "Toolkits/AssetEditorToolkit.h"
+#include "PhysicsPublic.h"
 
 /**
  * Public interface to Static Mesh Editor
@@ -11,6 +12,24 @@
 class IStaticMeshEditor : public FAssetEditorToolkit
 {
 public:
+	/**
+	 * Primitive data use to track which aggregate geometry is selected
+	 */
+	struct FPrimData
+	{
+		EKCollisionPrimitiveType		PrimType;
+		int32							PrimIndex;
+
+		FPrimData(EKCollisionPrimitiveType InPrimType, int32 InPrimIndex) :
+			PrimType(InPrimType),
+			PrimIndex(InPrimIndex) {}
+
+		bool operator==(const FPrimData& Other) const
+		{
+			return (PrimType == Other.PrimType && PrimIndex == Other.PrimIndex);
+		}
+	};
+
 	/** Called after an undo is performed to give child widgets a chance to refresh. */
 	DECLARE_MULTICAST_DELEGATE( FOnPostUndoMulticaster );
 
@@ -39,21 +58,122 @@ public:
 	/** Requests to rename selected socket */
 	virtual void RequestRenameSelectedSocket() = 0;
 
-	/** 	Retrieves the number of triangles in the current static mesh or it's forced LOD. 
+	/** 
+	 *  Checks to see if the prim data is valid compared with the static mesh
+	 *
+	 *  @param  InPrimData			The data to check
+	 */
+	virtual bool IsPrimValid(const FPrimData& InPrimData) const = 0;
+
+	/** Checks to see if any prims are selected */
+	virtual bool HasSelectedPrims() const = 0;
+
+	/** 
+	 *  Adds primitive information to the selected prims list
+	 *
+	 *  @param  InPrimData			The data to add
+	 */
+	virtual void AddSelectedPrim(const FPrimData& InPrimData) = 0;
+
+	/** 
+	 *  Removes primitive information to the selected prims list
+	 *
+	 *  @param  InPrimData			The data to remove
+	 */
+	virtual void RemoveSelectedPrim(const FPrimData& InPrimData) = 0;
+
+	/** Removes all invalid primitives from the list */
+	virtual void RemoveInvalidPrims() = 0;
+
+	/** 
+	 *  Checks to see if the parsed primitive data is selected
+	 *
+	 *  @param  InPrimData			The data to compare
+	 *  @returns					True, if the prim is selected
+	 */
+	virtual bool IsSelectedPrim(const FPrimData& InPrimData) const = 0;
+
+	/** Removes all primitive data from the list */
+	virtual void ClearSelectedPrims() = 0;
+
+	/**
+	 *  Duplicates all the selected primitives and selects them
+	 *
+	 *  @param  InOffset			[optional] Allows the duplicate to be offset by this factor
+	 */
+	virtual void DuplicateSelectedPrims(const FVector* InOffset) = 0;
+
+	/**
+	 *  Translates the selected primitives by the specified amount
+	 *
+	 *  @param  InDrag				The amount to translate
+	 */
+	virtual void TranslateSelectedPrims(const FVector& InDrag) = 0;
+
+	/**
+	 *  Rotates the selected primitives by the specified amount
+	 *
+	 *  @param  InRot				The amount to rotate
+	 */
+	virtual void RotateSelectedPrims(const FRotator& InRot) = 0;
+
+	/**
+	 *  Scales the selected primitives by the specified amount
+	 *
+	 *  @param  InScale				The amount to scale
+	 */
+	virtual void ScaleSelectedPrims(const FVector& InScale) = 0;
+
+	/**
+	 *  Calculates the bounding box of our selected primitives
+	 *
+	 *  @param  OutBox				The bounding data for our selection
+	 *  @returns					True, if there are any prims selected
+	 */
+	virtual bool CalcSelectedPrimsAABB(FBox &OutBox) const = 0;
+
+	/**
+	 *  Fetches the transform of the last primitive to be selected
+	 *
+	 *  @param  OutTransform		The transform of the last selected primitive
+	 *  @returns					True, if there was a prim selected
+	 */
+	virtual bool GetLastSelectedPrimTransform(FTransform& OutTransform) const = 0;
+
+	/**
+	 *  Gets the transform of the specified primitive
+	 *
+	 *  @param  InPrimData			The data about the primitive
+	 *  @returns					The transform of the specified primitive
+	 */
+	virtual FTransform GetPrimTransform(const FPrimData& InPrimData) const = 0;
+
+	/**
+	 *  Sets the transform of the specified primitive
+	 *
+	 *  @param  InPrimData			The data about the primitive
+	 *  @param  InPrimTransform		The transform to apply to the primitive
+	 */
+	virtual void SetPrimTransform(const FPrimData& InPrimData, const FTransform& InPrimTransform) const = 0;
+
+	/** 
+	 *  Retrieves the number of triangles in the current static mesh or it's forced LOD. 
 	 *
 	 *  @param  LODLevel			The desired LOD to retrieve the number of triangles for.
 	 *	@returns					The number of triangles for the specified LOD level.
 	 */
 	virtual int32 GetNumTriangles(int32 LODLevel = 0) const = 0;
 
-	/** Retrieves the number of vertices in the current static mesh or it's forced LOD. 
+	/** 
+	 *  Retrieves the number of vertices in the current static mesh or it's forced LOD. 
 	 *
 	 *  @param  LODLevel			The desired LOD to retrieve the number of vertices for.	 
 	 *	@returns					The number of vertices for the specified LOD level.
 	 */
 	virtual int32 GetNumVertices(int32 LODLevel = 0) const = 0;
 
-	/** Retrieves the number of UV channels available.
+	/** 
+	 *  Retrieves the number of UV channels available.
 	 *
 	 *  @param  LODLevel			The desired LOD to retrieve the number of UV channels for.
 	 *	@returns					The number of triangles for the specified LOD level.
diff --git a/Engine/Source/Editor/StaticMeshEditor/Public/StaticMeshEditorActions.h b/Engine/Source/Editor/StaticMeshEditor/Public/StaticMeshEditorActions.h
--- a/Engine/Source/Editor/StaticMeshEditor/Public/StaticMeshEditorActions.h
+++ b/Engine/Source/Editor/StaticMeshEditor/Public/StaticMeshEditorActions.h
@@ -43,13 +43,14 @@ public:
 	TSharedPtr< FUICommandInfo > SetShowPivot;
 
 	// Collision Menu Commands
-	TSharedPtr< FUICommandInfo > CreateDOP6;
 	TSharedPtr< FUICommandInfo > CreateDOP10X;
 	TSharedPtr< FUICommandInfo > CreateDOP10Y;
 	TSharedPtr< FUICommandInfo > CreateDOP10Z;
 	TSharedPtr< FUICommandInfo > CreateDOP18;
 	TSharedPtr< FUICommandInfo > CreateDOP26;
+	TSharedPtr< FUICommandInfo > CreateBoxCollision;
 	TSharedPtr< FUICommandInfo > CreateSphereCollision;
+	TSharedPtr< FUICommandInfo > CreateSphylCollision;
 	TSharedPtr< FUICommandInfo > CreateAutoConvexCollision;
 	TSharedPtr< FUICommandInfo > RemoveCollision;
 	TSharedPtr< FUICommandInfo > ConvertBoxesToConvex;
diff --git a/Engine/Source/Editor/UnrealEd/Private/ConvexDecompTool.cpp b/Engine/Source/Editor/UnrealEd/Private/ConvexDecompTool.cpp
--- a/Engine/Source/Editor/UnrealEd/Private/ConvexDecompTool.cpp
+++ b/Engine/Source/Editor/UnrealEd/Private/ConvexDecompTool.cpp
@@ -67,8 +67,7 @@ void DecomposeMeshToHulls(UBodySetup* InBodySetup, const TArray<FVector>& InVert
 	for(int32 HullIdx=0; HullIdx<NumHulls; HullIdx++)
 	{
 		// Create a new hull in the aggregate geometry
-		const int32 NewConvexIndex = InBodySetup->AggGeom.ConvexElems.Add(FKConvexElem());
-		FKConvexElem& Convex = InBodySetup->AggGeom.ConvexElems[NewConvexIndex];
+		FKConvexElem ConvexElem;
 
 		// Read out each hull vertex
 		const HACD_API::Hull* Hull = HAPI->getHull(HullIdx);
@@ -76,10 +75,12 @@ void DecomposeMeshToHulls(UBodySetup* InBodySetup, const TArray<FVector>& InVert
 		{
 			FVector* V = (FVector*)(Hull->mVertices + (VertIdx*3));
 
-			Convex.VertexData.Add(*V);
+			ConvexElem.VertexData.Add(*V);
 		}
 
-		Convex.UpdateElemBox();
+		ConvexElem.UpdateElemBox();
+
+		InBodySetup->AggGeom.ConvexElems.Add(ConvexElem);
 	}
 
 	InBodySetup->InvalidatePhysicsData(); // update GUID
diff --git a/Engine/Source/Editor/UnrealEd/Private/GeomFitUtils.cpp b/Engine/Source/Editor/UnrealEd/Private/GeomFitUtils.cpp
--- a/Engine/Source/Editor/UnrealEd/Private/GeomFitUtils.cpp
+++ b/Engine/Source/Editor/UnrealEd/Private/GeomFitUtils.cpp
@@ -29,33 +29,22 @@ static void AddVertexIfNotPresent(TArray<FVector> &vertices, FVector &newVertex)
 
 }
 
-/* ******************************** KDOP ******************************** */
-
-// This function takes the current collision model, and fits a k-DOP around it.
-// It uses the array of k unit-length direction vectors to define the k bounding planes.
-
-// THIS FUNCTION REPLACES EXISTING SIMPLE COLLISION MODEL WITH KDOP
-#define MY_FLTMAX (3.402823466e+38F)
-
-void GenerateKDopAsSimpleCollision(UStaticMesh* StaticMesh,TArray<FVector> &dirs)
+static bool PromptToRemoveExistingCollision(UStaticMesh* StaticMesh)
 {
-	// Make sure rendering is done - so we are not changing data being used by collision drawing.
-	FlushRenderingCommands();
-
+	check(StaticMesh);
 	UBodySetup* bs = StaticMesh->BodySetup;
 	if(bs && (bs->AggGeom.GetElementCount() > 0))
 	{
 		// If we already have some simplified collision for this mesh - check before we clobber it.
-		bool doReplace = EAppReturnType::Yes == FMessageDialog::Open( EAppMsgType::YesNo, NSLOCTEXT("UnrealEd", "StaticMeshAlreadyHasGeom", "Static Mesh already has simple collision.\nAre you sure you want replace it?") );
-
-		if(doReplace)
+		/*const EAppReturnType::Type ret = FMessageDialog::Open(EAppMsgType::YesNoCancel, NSLOCTEXT("UnrealEd", "StaticMeshAlreadyHasGeom", "Static Mesh already has simple collision.\nDo you want to replace it?"));
+		if (ret == EAppReturnType::Yes)
 		{
 			bs->RemoveSimpleCollision();
 		}
-		else
+		else if (ret == EAppReturnType::Cancel)
 		{
-			return;
-		}
+			return false;
+		}*/
 	}
 	else
 	{
@@ -63,9 +52,31 @@ void GenerateKDopAsSimpleCollision(UStaticMesh* StaticMesh,TArray<FVector> &dirs
 		StaticMesh->CreateBodySetup();
 		bs = StaticMesh->BodySetup;
 	}
+	return true;
+}
+
+/* ******************************** KDOP ******************************** */
+
+// This function takes the current collision model, and fits a k-DOP around it.
+// It uses the array of k unit-length direction vectors to define the k bounding planes.
+
+// THIS FUNCTION REPLACES EXISTING SIMPLE COLLISION MODEL WITH KDOP
+#define MY_FLTMAX (3.402823466e+38F)
+
+int32 GenerateKDopAsSimpleCollision(UStaticMesh* StaticMesh, TArray<FVector> &Dirs)
+{
+	// Make sure rendering is done - so we are not changing data being used by collision drawing.
+	FlushRenderingCommands();
+
+	if (!PromptToRemoveExistingCollision(StaticMesh))
+	{
+		return INDEX_NONE;
+	}
+
+	UBodySetup* bs = StaticMesh->BodySetup;
 
 	// Do k- specific stuff.
-	int32 kCount = dirs.Num();
+	int32 kCount = Dirs.Num();
 	TArray<float> maxDist;
 	for(int32 i=0; i<kCount; i++)
 		maxDist.Add(-MY_FLTMAX);
@@ -79,7 +90,7 @@ void GenerateKDopAsSimpleCollision(UStaticMesh* StaticMesh,TArray<FVector> &dirs
 	{
 		for(int32 j=0; j<kCount; j++)
 		{
-			float dist = RenderData.PositionVertexBuffer.VertexPosition(i) | dirs[j];
+			float dist = RenderData.PositionVertexBuffer.VertexPosition(i) | Dirs[j];
 			maxDist[j] = FMath::Max(dist, maxDist[j]);
 		}
 	}
@@ -98,7 +109,7 @@ void GenerateKDopAsSimpleCollision(UStaticMesh* StaticMesh,TArray<FVector> &dirs
 	// Now we have the planes of the kdop, we work out the face polygons.
 	TArray<FPlane> planes;
 	for(int32 i=0; i<kCount; i++)
-		planes.Add( FPlane(dirs[i], maxDist[i]) );
+		planes.Add( FPlane(Dirs[i], maxDist[i]) );
 
 	for(int32 i=0; i<planes.Num(); i++)
 	{
@@ -144,7 +155,7 @@ void GenerateKDopAsSimpleCollision(UStaticMesh* StaticMesh,TArray<FVector> &dirs
 	if(TempModel->Polys->Element.Num() < 4)
 	{
 		TempModel = NULL;
-		return;
+		return INDEX_NONE;
 	}
 
 	// Build bounding box.
@@ -155,13 +166,70 @@ void GenerateKDopAsSimpleCollision(UStaticMesh* StaticMesh,TArray<FVector> &dirs
 	FBSPOps::bspRefresh(TempModel,1);
 	FBSPOps::bspBuildBounds(TempModel);
 
-	bs->CreateFromModel(TempModel, true);
+	bs->Modify();
+
+	bs->CreateFromModel(TempModel, false);
 	
 	// create all body instances
 	RefreshCollisionChange(StaticMesh);
 
 	// Mark staticmesh as dirty, to help make sure it gets saved.
 	StaticMesh->MarkPackageDirty();
+
+	return bs->AggGeom.ConvexElems.Num() - 1;
+}
+
+/* ******************************** BOX ******************************** */
+
+static void CalcBoundingBox(const FRawMesh& RawMesh, FVector& Center, FVector& Extents, FVector& LimitVec)
+{
+	FBox Box(0);
+
+	for (uint32 i = 0; i < (uint32)RawMesh.VertexPositions.Num(); i++)
+	{
+		Box += RawMesh.VertexPositions[i] * LimitVec;
+	}
+
+	Box.GetCenterAndExtents(Center, Extents);
+}
+
+int32 GenerateBoxAsSimpleCollision(UStaticMesh* StaticMesh)
+{
+	if (!PromptToRemoveExistingCollision(StaticMesh))
+	{
+		return INDEX_NONE;
+	}
+
+	UBodySetup* bs = StaticMesh->BodySetup;
+
+	// Calculate bounding Box.
+	FRawMesh RawMesh;
+	FStaticMeshSourceModel& SrcModel = StaticMesh->SourceModels[0];
+	SrcModel.RawMeshBulkData->LoadRawMesh(RawMesh);
+
+	FVector unitVec = bs->BuildScale3D;
+	FVector Center, Extents;
+	CalcBoundingBox(RawMesh, Center, Extents, unitVec);
+
+	bs->Modify();
+
+	// Create new GUID
+	bs->InvalidatePhysicsData();
+
+	FKBoxElem BoxElem;
+	BoxElem.Center = Center;
+	BoxElem.X = Extents.X * 2.0f;
+	BoxElem.Y = Extents.Y * 2.0f;
+	BoxElem.Z = Extents.Z * 2.0f;
+	bs->AggGeom.BoxElems.Add(BoxElem);
+
+	// refresh collision change back to staticmesh components
+	RefreshCollisionChange(StaticMesh);
+
+	// Mark staticmesh as dirty, to help make sure it gets saved.
+	StaticMesh->MarkPackageDirty();
+
+	return bs->AggGeom.BoxElems.Num() - 1;
 }
 
 /* ******************************** SPHERE ******************************** */
@@ -173,19 +241,19 @@ void GenerateKDopAsSimpleCollision(UStaticMesh* StaticMesh,TArray<FVector> &dirs
 // This one seems to do well with asymmetric input.
 static void CalcBoundingSphere(const FRawMesh& RawMesh, FSphere& sphere, FVector& LimitVec)
 {
-	FBox Box;
-
 	if(RawMesh.VertexPositions.Num() == 0)
 		return;
 
-	int32 minIx[3], maxIx[3]; // Extreme points.
+	FBox Box;
 
 	// First, find AABB, remembering furthest points in each dir.
 	Box.Min = RawMesh.VertexPositions[0] * LimitVec;
 	Box.Max = Box.Min;
 
-	minIx[0] = minIx[1] = minIx[2] = 0;
-	maxIx[0] = maxIx[1] = maxIx[2] = 0;
+	FIntVector minIx, maxIx; // Extreme points.
+
+	minIx = FIntVector::ZeroValue;
+	maxIx = FIntVector::ZeroValue;
 
 	for(uint32 i=1; i<(uint32)RawMesh.VertexPositions.Num(); i++) 
 	{
@@ -195,75 +263,76 @@ static void CalcBoundingSphere(const FRawMesh& RawMesh, FSphere& sphere, FVector
 		if(p.X < Box.Min.X)
 		{
 			Box.Min.X = p.X;
-			minIx[0] = i;
+			minIx.X = i;
 		}
 		else if(p.X > Box.Max.X)
 		{
 			Box.Max.X = p.X;
-			maxIx[0] = i;
+			maxIx.X = i;
 		}
 
 		// Y //
 		if(p.Y < Box.Min.Y)
 		{
 			Box.Min.Y = p.Y;
-			minIx[1] = i;
+			minIx.Y = i;
 		}
 		else if(p.Y > Box.Max.Y)
 		{
 			Box.Max.Y = p.Y;
-			maxIx[1] = i;
+			maxIx.Y = i;
 		}
 
 		// Z //
 		if(p.Z < Box.Min.Z)
 		{
 			Box.Min.Z = p.Z;
-			minIx[2] = i;
+			minIx.Z = i;
 		}
 		else if(p.Z > Box.Max.Z)
 		{
 			Box.Max.Z = p.Z;
-			maxIx[2] = i;
+			maxIx.Z = i;
 		}
 	}
 
-	//  Now find extreme points furthest apart, and initial centre and radius of sphere.
+	const FVector Extremes[3]={ (RawMesh.VertexPositions[maxIx.X] - RawMesh.VertexPositions[minIx.X]) * LimitVec,
+								(RawMesh.VertexPositions[maxIx.Y] - RawMesh.VertexPositions[minIx.Y]) * LimitVec,
+								(RawMesh.VertexPositions[maxIx.Z] - RawMesh.VertexPositions[minIx.Z]) * LimitVec };
+
+	// Now find extreme points furthest apart, and initial center and radius of sphere.
 	float d2 = 0.f;
 	for(int32 i=0; i<3; i++)
 	{
-		FVector diff = (RawMesh.VertexPositions[maxIx[i]] - RawMesh.VertexPositions[minIx[i]]) * LimitVec;
-		float tmpd2 = diff.SizeSquared();
-
+		const float tmpd2 = Extremes[i].SizeSquared();
 		if(tmpd2 > d2)
 		{
 			d2 = tmpd2;
-			FVector centre = RawMesh.VertexPositions[minIx[i]] + (0.5f * diff);
-			centre *= LimitVec;
-			sphere.Center.X = centre.X;
-			sphere.Center.Y = centre.Y;
-			sphere.Center.Z = centre.Z;
+			sphere.Center = (RawMesh.VertexPositions[minIx(i)] + (0.5f * Extremes[i])) * LimitVec;
 			sphere.W = 0.f;
 		}
 	}
 
+	const FVector Extents = FVector(Extremes[0].X, Extremes[1].Y, Extremes[2].Z);
+
 	// radius and radius squared
-	float r = 0.5f * FMath::Sqrt(d2);
-	float r2 = r * r;
+	float r = 0.5f * Extents.GetMax();
+	float r2 = FMath::Square(r);
 
 	// Now check each point lies within this sphere. If not - expand it a bit.
 	for(uint32 i=0; i<(uint32)RawMesh.VertexPositions.Num(); i++) 
 	{
-		FVector cToP = (RawMesh.VertexPositions[i] * LimitVec) - sphere.Center;
-		float pr2 = cToP.SizeSquared();
+		const FVector cToP = (RawMesh.VertexPositions[i] * LimitVec) - sphere.Center;
 
-		// If this point is outside our current bounding sphere..
+		const float pr2 = cToP.SizeSquared();
+
+		// If this point is outside our current bounding sphere's radius
 		if(pr2 > r2)
 		{
-			// ..expand sphere just enough to include this point.
-			float pr = FMath::Sqrt(pr2);
+			// ..expand radius just enough to include this point.
+			const float pr = FMath::Sqrt(pr2);
 			r = 0.5f * (r + pr);
-			r2 = r * r;
+			r2 = FMath::Square(r);
 
 			sphere.Center += ((pr-r)/pr * cToP);
 		}
@@ -276,20 +345,11 @@ static void CalcBoundingSphere(const FRawMesh& RawMesh, FSphere& sphere, FVector
 // Seems to do better with more symmetric input...
 static void CalcBoundingSphere2(const FRawMesh& RawMesh, FSphere& sphere, FVector& LimitVec)
 {
-	FBox Box(0);
-	
-	for(uint32 i=0; i<(uint32)RawMesh.VertexPositions.Num(); i++)
-	{
-		Box += RawMesh.VertexPositions[i] * LimitVec;
-	}
+	FVector Center, Extents;
+	CalcBoundingBox(RawMesh, Center, Extents, LimitVec);
 
-	FVector centre, extent;
-	Box.GetCenterAndExtents(centre, extent);
-
-	sphere.Center.X = centre.X;
-	sphere.Center.Y = centre.Y;
-	sphere.Center.Z = centre.Z;
-	sphere.W = 0;
+	sphere.Center = Center;
+	sphere.W = 0.0f;
 
 	for( uint32 i=0; i<(uint32)RawMesh.VertexPositions.Num(); i++ )
 	{
@@ -302,36 +362,22 @@ static void CalcBoundingSphere2(const FRawMesh& RawMesh, FSphere& sphere, FVecto
 
 // // //
 
-void GenerateSphereAsSimpleCollision(UStaticMesh* StaticMesh)
+int32 GenerateSphereAsSimpleCollision(UStaticMesh* StaticMesh)
 {
-	UBodySetup* bs = StaticMesh->BodySetup;
-	if(bs && (bs->AggGeom.GetElementCount() > 0))
-	{
-		bool doReplace = EAppReturnType::Yes == FMessageDialog::Open( EAppMsgType::YesNo, NSLOCTEXT("UnrealEd", "StaticMeshAlreadyHasGeom", "Static Mesh already has simple collision.\nAre you sure you want replace it?") );
-
-		if(doReplace)
-		{
-			bs->RemoveSimpleCollision();
-		}
-		else
-		{
-			return;
-		}
-	}
-	else
+	if (!PromptToRemoveExistingCollision(StaticMesh))
 	{
-		// Otherwise, create one here.
-		StaticMesh->CreateBodySetup();
-		bs = StaticMesh->BodySetup;
+		return INDEX_NONE;
 	}
 
+	UBodySetup* bs = StaticMesh->BodySetup;
+
 	// Calculate bounding sphere.
 	FRawMesh RawMesh;
 	FStaticMeshSourceModel& SrcModel = StaticMesh->SourceModels[0];
 	SrcModel.RawMeshBulkData->LoadRawMesh(RawMesh);
 
 	FSphere bSphere, bSphere2, bestSphere;
-	FVector unitVec = StaticMesh->BodySetup ? StaticMesh->BodySetup->BuildScale3D : FVector(1.f, 1.f, 1.f);
+	FVector unitVec = bs->BuildScale3D;
 	CalcBoundingSphere(RawMesh, bSphere, unitVec);
 	CalcBoundingSphere2(RawMesh, bSphere2, unitVec);
 
@@ -344,39 +390,176 @@ void GenerateSphereAsSimpleCollision(UStaticMesh* StaticMesh)
 	if(bestSphere.W <= 0.f)
 	{
 		FMessageDialog::Open( EAppMsgType::Ok, NSLOCTEXT("UnrealEd", "Prompt_10", "Could not create geometry.") );
-		return;
+		return INDEX_NONE;
 	}
 
-	int ex = bs->AggGeom.SphereElems.AddZeroed();
-	FKSphereElem* s = &bs->AggGeom.SphereElems[ex];
-	s->Center = bestSphere.Center;
-	s->Radius = bestSphere.W;
+	bs->Modify();
 
 	// Create new GUID
 	bs->InvalidatePhysicsData();
 
+	FKSphereElem SphereElem;
+	SphereElem.Center = bestSphere.Center;
+	SphereElem.Radius = bestSphere.W;
+	bs->AggGeom.SphereElems.Add(SphereElem);
+
 	// refresh collision change back to staticmesh components
 	RefreshCollisionChange(StaticMesh);
 
 	// Mark staticmesh as dirty, to help make sure it gets saved.
 	StaticMesh->MarkPackageDirty();
+
+	return bs->AggGeom.SphereElems.Num() - 1;
+}
+
+/* ******************************** SPHYL ******************************** */
+
+static void CalcBoundingSphyl(const FRawMesh& RawMesh, FSphere& sphere, float& length, FRotator& rotation, FVector& LimitVec)
+{
+	if (RawMesh.VertexPositions.Num() == 0)
+		return;
+
+	FVector Center, Extents;
+	CalcBoundingBox(RawMesh, Center, Extents, LimitVec);
+
+	// @todo sphere.Center could perhaps be adjusted to best fit if model is non-symmetric on it's longest axis
+	sphere.Center = Center;
+
+	// Work out best axis aligned orientation (longest side)
+	float Extent = Extents.GetMax();
+	if (Extent == Extents.X)
+	{
+		rotation = FRotator(90.f, 0.f, 0.f);
+		Extents.X = 0.0f;
+	}
+	else if (Extent == Extents.Y)
+	{
+		rotation = FRotator(0.f, 0.f, 90.f);
+		Extents.Y = 0.0f;
+	}
+	else
+	{
+		rotation = FRotator(0.f, 0.f, 0.f);
+		Extents.Z = 0.0f;
+	}
+
+	// Cleared the largest axis above, remaining determines the radius
+	float r = Extents.GetMax();
+	float r2 = FMath::Square(r);
+	
+	// Now check each point lies within this the radius. If not - expand it a bit.
+	for (uint32 i = 0; i<(uint32)RawMesh.VertexPositions.Num(); i++)
+	{
+		FVector cToP = (RawMesh.VertexPositions[i] * LimitVec) - sphere.Center;
+		cToP = rotation.UnrotateVector(cToP);
+
+		const float pr2 = cToP.SizeSquared2D();	// Ignore Z here...
+
+		// If this point is outside our current bounding sphere's radius
+		if (pr2 > r2)
+		{
+			// ..expand radius just enough to include this point.
+			const float pr = FMath::Sqrt(pr2);
+			r = 0.5f * (r + pr);
+			r2 = FMath::Square(r);
+		}
+	}
+	
+	// The length is the longest side minus the radius.
+	float hl = FMath::Max(0.0f, Extent - r);
+
+	// Now check each point lies within the length. If not - expand it a bit.
+	for (uint32 i = 0; i<(uint32)RawMesh.VertexPositions.Num(); i++)
+	{
+		FVector cToP = (RawMesh.VertexPositions[i] * LimitVec) - sphere.Center;
+		cToP = rotation.UnrotateVector(cToP);
+
+		// If this point is outside our current bounding sphyl's length
+		if (FMath::Abs(cToP.Z) > hl)
+		{
+			const bool bFlip = (cToP.Z < 0.f ? true : false);
+			const FVector cOrigin(0.f, 0.f, (bFlip ? -hl : hl));
+
+			const float pr2 = (cOrigin - cToP).SizeSquared();
+
+			// If this point is outside our current bounding sphyl's radius
+			if (pr2 > r2)
+			{
+				FVector cPoint;
+				FMath::SphereDistToLine(cOrigin, r, cToP, (bFlip ? FVector(0.f, 0.f, 1.f) : FVector(0.f, 0.f, -1.f)), cPoint);
+
+				// Don't accept zero as a valid diff when we know it's outside the sphere (saves needless retest on further iterations of like points)
+				hl += FMath::Max(FMath::Abs(cToP.Z - cPoint.Z), 1.e-6f);
+			}
+		}
+	}
+
+	sphere.W = r;
+	length = hl * 2.0f;
 }
 
+// // //
+
+int32 GenerateSphylAsSimpleCollision(UStaticMesh* StaticMesh)
+{
+	if (!PromptToRemoveExistingCollision(StaticMesh))
+	{
+		return INDEX_NONE;
+	}
+
+	UBodySetup* bs = StaticMesh->BodySetup;
+
+	// Calculate bounding box.
+	FRawMesh RawMesh;
+	FStaticMeshSourceModel& SrcModel = StaticMesh->SourceModels[0];
+	SrcModel.RawMeshBulkData->LoadRawMesh(RawMesh);
+
+	FSphere sphere;
+	float length;
+	FRotator rotation;
+	FVector unitVec = bs->BuildScale3D;
+	CalcBoundingSphyl(RawMesh, sphere, length, rotation, unitVec);
+
+	// Dont use if radius is zero.
+	if (sphere.W <= 0.f)
+	{
+		FMessageDialog::Open(EAppMsgType::Ok, NSLOCTEXT("UnrealEd", "Prompt_10", "Could not create geometry."));
+		return INDEX_NONE;
+	}
+
+	// If height is zero, then a sphere would be better (should we just create one instead?)
+	if (length <= 0.f)
+	{
+		length = SMALL_NUMBER;
+	}
+
+	bs->Modify();
+
+	// Create new GUID
+	bs->InvalidatePhysicsData();
+
+	FKSphylElem SphylElem;
+	SphylElem.Center = sphere.Center;
+	SphylElem.Orientation = rotation.Quaternion();
+	SphylElem.Radius = sphere.W;
+	SphylElem.Length = length;
+	bs->AggGeom.SphylElems.Add(SphylElem);
+
+	// refresh collision change back to staticmesh components
+	RefreshCollisionChange(StaticMesh);
+
+	// Mark staticmesh as dirty, to help make sure it gets saved.
+	StaticMesh->MarkPackageDirty();
+
+	return bs->AggGeom.SphylElems.Num() - 1;
+}
 
-/**
- * Refresh Collision Change
- * 
- * Collision has been changed, so it will need to recreate physics state to reflect it
- * Utilities functions to propagate BodySetup change for StaticMesh
- *
- * @param	InStaticMesh	StaticMesh that collision has been changed for
- */
-void RefreshCollisionChange(const UStaticMesh * InStaticMesh)
+void RefreshCollisionChange(const UStaticMesh * StaticMesh)
 {
 	for (FObjectIterator Iter(UStaticMeshComponent::StaticClass()); Iter; ++Iter)
 	{
 		UStaticMeshComponent* StaticMeshComponent = Cast<UStaticMeshComponent>(*Iter);
-		if  (StaticMeshComponent->StaticMesh == InStaticMesh)
+		if  (StaticMeshComponent->StaticMesh == StaticMesh)
 		{
 			// it needs to recreate IF it already has been created
 			if (StaticMeshComponent->IsPhysicsStateCreated())
diff --git a/Engine/Source/Editor/UnrealEd/Private/GeomFitUtils.h b/Engine/Source/Editor/UnrealEd/Private/GeomFitUtils.h
--- a/Engine/Source/Editor/UnrealEd/Private/GeomFitUtils.h
+++ b/Engine/Source/Editor/UnrealEd/Private/GeomFitUtils.h
@@ -10,16 +10,6 @@
 #define RCP_SQRT3 (0.57735026918962576450914878050196f)
 
 
-const FVector KDopDir6[6] = 
-{
-	FVector( 1.f, 0.f, 0.f),
-	FVector(-1.f, 0.f, 0.f),
-	FVector( 0.f, 1.f, 0.f),
-	FVector( 0.f,-1.f, 0.f),
-	FVector( 0.f, 0.f, 1.f),
-	FVector( 0.f, 0.f,-1.f)
-};
-
 const FVector KDopDir10X[10] = 
 {
 	FVector( 1.f, 0.f, 0.f),
@@ -115,8 +105,10 @@ const FVector KDopDir26[26] =
 };
 
 // Utilities
-UNREALED_API void GenerateKDopAsSimpleCollision(UStaticMesh* StaticMesh,TArray<FVector> &dirs);
-UNREALED_API void GenerateSphereAsSimpleCollision(UStaticMesh* StaticMesh);
+UNREALED_API int32 GenerateKDopAsSimpleCollision(UStaticMesh* StaticMesh, TArray<FVector> &dirs);
+UNREALED_API int32 GenerateBoxAsSimpleCollision(UStaticMesh* StaticMesh);
+UNREALED_API int32 GenerateSphereAsSimpleCollision(UStaticMesh* StaticMesh);
+UNREALED_API int32 GenerateSphylAsSimpleCollision(UStaticMesh* StaticMesh);
 
 /**
  * Refresh Collision Change
@@ -124,6 +116,7 @@ UNREALED_API void GenerateSphereAsSimpleCollision(UStaticMesh* StaticMesh);
  * Collision has been changed, so it will need to recreate physics state to reflect it
  * Utilities functions to propagate BodySetup change for StaticMesh
  *
- * @param	InStaticMesh	StaticMesh that collision has been changed for
+ * @param	StaticMesh	StaticMesh that collision has been changed for
  */
-UNREALED_API void RefreshCollisionChange(const UStaticMesh * InStaticMesh);
\ No newline at end of file
+UNREALED_API void RefreshCollisionChange(const UStaticMesh * StaticMesh);
+	
\ No newline at end of file
diff --git a/Engine/Source/Editor/UnrealEd/Private/PhysicsAssetUtils.cpp b/Engine/Source/Editor/UnrealEd/Private/PhysicsAssetUtils.cpp
--- a/Engine/Source/Editor/UnrealEd/Private/PhysicsAssetUtils.cpp
+++ b/Engine/Source/Editor/UnrealEd/Private/PhysicsAssetUtils.cpp
@@ -313,36 +313,38 @@ void CreateCollisionFromBone( UBodySetup* bs, USkeletalMesh* skelMesh, int32 Bon
 	if(Params.GeomType == EFG_Box)
 	{
 		// Add a new box geometry to this body the size of the bounding box.
-		int32 ex = bs->AggGeom.BoxElems.AddZeroed();
-		FKBoxElem* be = &bs->AggGeom.BoxElems[ex];
+		FKBoxElem BoxElem;
 
-		be->SetTransform( FTransform( ElemTM ) );
+		BoxElem.SetTransform(FTransform(ElemTM));
 
-		be->X = BoxExtent.X * 2.0f * 1.01f; // Side Lengths (add 1% to avoid graphics glitches)
-		be->Y = BoxExtent.Y * 2.0f * 1.01f;
-		be->Z = BoxExtent.Z * 2.0f * 1.01f;	
+		BoxElem.X = BoxExtent.X * 2.0f * 1.01f; // Side Lengths (add 1% to avoid graphics glitches)
+		BoxElem.Y = BoxExtent.Y * 2.0f * 1.01f;
+		BoxElem.Z = BoxExtent.Z * 2.0f * 1.01f;
+
+		bs->AggGeom.BoxElems.Add(BoxElem);
 	}
 	else if (Params.GeomType == EFG_Sphere)
 	{
-		int32 sx = bs->AggGeom.SphereElems.AddZeroed();
-		FKSphereElem* se = &bs->AggGeom.SphereElems[sx];
+		FKSphereElem SphereElem;
+
+		SphereElem.Center = ElemTM.GetOrigin();
+		SphereElem.Radius = BoxExtent.GetMax() * 1.01f;
 
-		se->Center = ElemTM.GetOrigin();
-		se->Radius = BoxExtent.GetMax() * 1.01f;
+		bs->AggGeom.SphereElems.Add(SphereElem);
 	}
 	// Deal with creating a single convex hull
 	else if (Params.GeomType == EFG_SingleConvexHull)
 	{
-		int32 ex = bs->AggGeom.ConvexElems.Add(FKConvexElem());
-		FKConvexElem* ce = &bs->AggGeom.ConvexElems[ex];
+		FKConvexElem ConvexElem;
 
 		// Add all of the vertices for this bone to the convex element
 		for( int32 index=0; index<BoneInfo->Positions.Num(); index++ )
 		{
-			ce->VertexData.Add( BoneInfo->Positions[index] );
+			ConvexElem.VertexData.Add(BoneInfo->Positions[index]);
 		}
+		ConvexElem.UpdateElemBox();
 
-		ce->UpdateElemBox();		
+		bs->AggGeom.ConvexElems.Add(ConvexElem);
 	}
 	else if (Params.GeomType == EFG_MultiConvexHull)
 	{
@@ -408,13 +410,15 @@ void CreateCollisionFromBone( UBodySetup* bs, USkeletalMesh* skelMesh, int32 Bon
 	}
 	else
 	{
-		int32 sx = bs->AggGeom.SphylElems.AddZeroed();
-		FKSphylElem* se = &bs->AggGeom.SphylElems[sx];
+		
+		FKSphylElem SphylElem;
+
+		SphylElem.SetTransform(FTransform(ElemTM));
 
-		se->SetTransform( FTransform( ElemTM ) );
+		SphylElem.Radius = FMath::Max(BoxExtent.X, BoxExtent.Y) * 1.01f;
+		SphylElem.Length = BoxExtent.Z * 1.01f;
 
-		se->Radius = FMath::Max(BoxExtent.X, BoxExtent.Y) * 1.01f;
-		se->Length = BoxExtent.Z * 1.01f;
+		bs->AggGeom.SphylElems.Add(SphylElem);
 	}
 }
 
diff --git a/Engine/Source/Editor/UnrealEd/Private/SSocketManager.cpp b/Engine/Source/Editor/UnrealEd/Private/SSocketManager.cpp
--- a/Engine/Source/Editor/UnrealEd/Private/SSocketManager.cpp
+++ b/Engine/Source/Editor/UnrealEd/Private/SSocketManager.cpp
@@ -371,19 +371,30 @@ UStaticMeshSocket* SSocketManager::GetSelectedSocket() const
 
 void SSocketManager::SetSelectedSocket(UStaticMeshSocket* InSelectedSocket)
 {
-	for( int32 i=0; i < SocketList.Num(); i++)
+	if (InSelectedSocket)
 	{
-		if(SocketList[i]->Socket == InSelectedSocket)
+		for( int32 i=0; i < SocketList.Num(); i++)
 		{
-			SocketListView->SetSelection(SocketList[i]);
+			if(SocketList[i]->Socket == InSelectedSocket)
+			{
+				SocketListView->SetSelection(SocketList[i]);
 
-			SocketListView->RequestListRefresh();
+				SocketListView->RequestListRefresh();
 
-			SocketSelectionChanged(InSelectedSocket);
+				SocketSelectionChanged(InSelectedSocket);
 
-			return;
+				break;
+			}
 		}
 	}
+	else
+	{
+		SocketListView->ClearSelection();
+
+		SocketListView->RequestListRefresh();
+
+		SocketSelectionChanged(NULL);
+	}
 }
 
 TSharedRef< ITableRow > SSocketManager::MakeWidgetFromOption( TSharedPtr<SocketListItem> InItem, const TSharedRef< STableViewBase >& OwnerTable )
diff --git a/Engine/Source/Editor/UnrealEd/Private/StaticMeshEdit.cpp b/Engine/Source/Editor/UnrealEd/Private/StaticMeshEdit.cpp
--- a/Engine/Source/Editor/UnrealEd/Private/StaticMeshEdit.cpp
+++ b/Engine/Source/Editor/UnrealEd/Private/StaticMeshEdit.cpp
@@ -411,14 +411,13 @@ void AddBoxGeomFromTris( const TArray<FPoly>& Tris, FKAggregateGeom* AggGeom, co
 	BoxTM.SetOrigin( Box.GetCenter() );
 
 	// Allocate box in array
-	int32 NewIndex = AggGeom->BoxElems.AddZeroed();
-	FKBoxElem* BoxElem = &AggGeom->BoxElems[NewIndex];
-
-	BoxElem->SetTransform( FTransform( BoxTM ) );	
+	FKBoxElem BoxElem;
+	BoxElem.SetTransform( FTransform( BoxTM ) );	
 	// distance between parallel planes is box edge lengths.
-	BoxElem->X = FMath::Abs(Planes[0].PlaneDist[0] - Planes[0].PlaneDist[1]);
-	BoxElem->Y = FMath::Abs(Planes[1].PlaneDist[0] - Planes[1].PlaneDist[1]);
-	BoxElem->Z = FMath::Abs(Planes[2].PlaneDist[0] - Planes[2].PlaneDist[1]);
+	BoxElem.X = FMath::Abs(Planes[0].PlaneDist[0] - Planes[0].PlaneDist[1]);
+	BoxElem.Y = FMath::Abs(Planes[1].PlaneDist[0] - Planes[1].PlaneDist[1]);
+	BoxElem.Z = FMath::Abs(Planes[2].PlaneDist[0] - Planes[2].PlaneDist[1]);
+	AggGeom->BoxElems.Add(BoxElem);
 }
 
 /**
@@ -483,12 +482,10 @@ void AddSphereGeomFromVerts( const TArray<FVector>& Verts, FKAggregateGeom* AggG
 	}
 
 	// Allocate sphere in array
-	int32 NewIndex = AggGeom->SphereElems.AddZeroed();
-	FKSphereElem* SphereElem = &AggGeom->SphereElems[NewIndex];
-
-	SphereElem->Center = Center;
-
-	SphereElem->Radius = Radius;
+	FKSphereElem SphereElem;
+	SphereElem.Center = Center;
+	SphereElem.Radius = Radius;
+	AggGeom->SphereElems.Add(SphereElem);
 }
 
 /** Utility for adding one convex hull from the given verts */
diff --git a/Engine/Source/Runtime/Engine/Classes/PhysicsEngine/BodySetup.h b/Engine/Source/Runtime/Engine/Classes/PhysicsEngine/BodySetup.h
--- a/Engine/Source/Runtime/Engine/Classes/PhysicsEngine/BodySetup.h
+++ b/Engine/Source/Runtime/Engine/Classes/PhysicsEngine/BodySetup.h
@@ -67,50 +67,51 @@ struct FKConvexElem
 	UPROPERTY()
 	FTransform Transform;
 
-		/** Convex mesh for this body, created from cooked data in CreatePhysicsMeshes */
-		physx::PxConvexMesh*   ConvexMesh;	
+	/** Convex mesh for this body, created from cooked data in CreatePhysicsMeshes */
+	physx::PxConvexMesh*   ConvexMesh;	
 
-		/** Convex mesh for this body, flipped across X, created from cooked data in CreatePhysicsMeshes */
-		physx::PxConvexMesh*   ConvexMeshNegX;
+	/** Convex mesh for this body, flipped across X, created from cooked data in CreatePhysicsMeshes */
+	physx::PxConvexMesh*   ConvexMeshNegX;
 
-		FKConvexElem() 
-		: ElemBox(0)
-		, Transform(FTransform::Identity)
-		, ConvexMesh(NULL)
-		, ConvexMeshNegX(NULL)
-		{}
+	FKConvexElem() 
+	: ElemBox(0)
+	, Transform(FTransform::Identity)
+	, ConvexMesh(NULL)
+	, ConvexMeshNegX(NULL)
+	{}
 
-		ENGINE_API void	DrawElemWire(class FPrimitiveDrawInterface* PDI, const FTransform& ElemTM, const FColor Color);
-		void	AddCachedSolidConvexGeom(TArray<FDynamicMeshVertex>& VertexBuffer, TArray<int32>& IndexBuffer, const FColor VertexColor);
+	ENGINE_API void	DrawElemWire(class FPrimitiveDrawInterface* PDI, const FTransform& ElemTM, const FColor Color) const;
+	void	AddCachedSolidConvexGeom(TArray<FDynamicMeshVertex>& VertexBuffer, TArray<int32>& IndexBuffer, const FColor VertexColor);
 
-		/** Reset the hull to empty all arrays */
-		ENGINE_API void	Reset();
+	/** Reset the hull to empty all arrays */
+	ENGINE_API void	Reset();
 
-		/** Updates internal ElemBox based on current value of VertexData */
-		ENGINE_API void	UpdateElemBox();
+	/** Updates internal ElemBox based on current value of VertexData */
+	ENGINE_API void	UpdateElemBox();
 
-		/** Calculate a bounding box for this convex element with the specified transform and scale */
-		FBox	CalcAABB(const FTransform& BoneTM, const FVector& Scale3D);		
+	/** Calculate a bounding box for this convex element with the specified transform and scale */
+	ENGINE_API FBox	CalcAABB(const FTransform& BoneTM, const FVector& Scale3D) const;
 
-		/** Utility for creating a convex hull from a set of planes. Will reset current state of this elem. */
-		bool	HullFromPlanes(const TArray<FPlane>& InPlanes, const TArray<FVector>& SnapVerts);
+	/** Utility for creating a convex hull from a set of planes. Will reset current state of this elem. */
+	bool	HullFromPlanes(const TArray<FPlane>& InPlanes, const TArray<FVector>& SnapVerts);
 
-		/** Returns the volume of this element */
-		float GetVolume(const FVector& Scale) const;
+	/** Returns the volume of this element */
+	float GetVolume(const FVector& Scale) const;
 
-		FTransform GetTransform() const
-		{
-			return Transform;
-		};
+	FTransform GetTransform() const
+	{
+		return Transform;
+	};
 
-		void SetTransform( const FTransform& InTransform )
-		{
-			ensure(InTransform.IsValid());
-			Transform = InTransform;
-		}
+	void SetTransform( const FTransform& InTransform )
+	{
+		ensure(InTransform.IsValid());
+		Transform = InTransform;
+	}
 
-		friend FArchive& operator<<(FArchive& Ar,FKConvexElem& Elem);
+	friend FArchive& operator<<(FArchive& Ar,FKConvexElem& Elem);
 	
+	ENGINE_API void ScaleElem(FVector DeltaSize, float MinSize);
 };
 
 /** Sphere shape used for collision */
@@ -156,11 +157,19 @@ struct FKSphereElem
 		return FTransform( Center );
 	};
 
+	void SetTransform(const FTransform& InTransform)
+	{
+		ensure(InTransform.IsValid());
+		Center = InTransform.GetLocation();
+	}
+
 	FORCEINLINE float GetVolume(const FVector& Scale) const { return 1.3333f * PI * FMath::Pow(Radius * Scale.GetMin(), 3); }
 	
-	ENGINE_API void	DrawElemWire(class FPrimitiveDrawInterface* PDI, const FTransform& ElemTM, float Scale, const FColor Color);
-	ENGINE_API void	DrawElemSolid(class FPrimitiveDrawInterface* PDI, const FTransform& ElemTM, float Scale, const FMaterialRenderProxy* MaterialRenderProxy);
-	FBox CalcAABB(const FTransform& BoneTM, float Scale);
+	ENGINE_API void	DrawElemWire(class FPrimitiveDrawInterface* PDI, const FTransform& ElemTM, float Scale, const FColor Color) const;
+	ENGINE_API void	DrawElemSolid(class FPrimitiveDrawInterface* PDI, const FTransform& ElemTM, float Scale, const FMaterialRenderProxy* MaterialRenderProxy) const;
+	ENGINE_API FBox CalcAABB(const FTransform& BoneTM, float Scale) const;
+
+	ENGINE_API void ScaleElem(FVector DeltaSize, float MinSize);
 };
 
 /** Box shape used for collision */
@@ -240,9 +249,11 @@ struct FKBoxElem
 
 	FORCEINLINE float GetVolume(const FVector& Scale) const { float MinScale = Scale.GetMin(); return (X * MinScale) * (Y * MinScale) * (Z * MinScale); }
 
-	ENGINE_API void	DrawElemWire(class FPrimitiveDrawInterface* PDI, const FTransform& ElemTM, float Scale, const FColor Color);
-	ENGINE_API void	DrawElemSolid(class FPrimitiveDrawInterface* PDI, const FTransform& ElemTM, float Scale, const FMaterialRenderProxy* MaterialRenderProxy);
-	FBox CalcAABB(const FTransform& BoneTM, float Scale);
+	ENGINE_API void	DrawElemWire(class FPrimitiveDrawInterface* PDI, const FTransform& ElemTM, float Scale, const FColor Color) const;
+	ENGINE_API void	DrawElemSolid(class FPrimitiveDrawInterface* PDI, const FTransform& ElemTM, float Scale, const FMaterialRenderProxy* MaterialRenderProxy) const;
+	ENGINE_API FBox CalcAABB(const FTransform& BoneTM, float Scale) const;
+
+	ENGINE_API void ScaleElem(FVector DeltaSize, float MinSize);
 };
 
 /** Capsule shape used for collision */
@@ -310,9 +321,11 @@ struct FKSphylElem
 
 	FORCEINLINE float GetVolume(const FVector& Scale) const { float ScaledRadius = Radius * Scale.GetMin(); return PI * FMath::Square(ScaledRadius) * ( 1.3333f * ScaledRadius + (Length * Scale.GetMin())); }
 
-	ENGINE_API void	DrawElemWire(class FPrimitiveDrawInterface* PDI, const FTransform& ElemTM, float Scale, const FColor Color);
-	ENGINE_API void	DrawElemSolid(class FPrimitiveDrawInterface* PDI, const FTransform& ElemTM, float Scale, const FMaterialRenderProxy* MaterialRenderProxy);
-	FBox CalcAABB(const FTransform& BoneTM, float Scale);
+	ENGINE_API void	DrawElemWire(class FPrimitiveDrawInterface* PDI, const FTransform& ElemTM, float Scale, const FColor Color) const;
+	ENGINE_API void	DrawElemSolid(class FPrimitiveDrawInterface* PDI, const FTransform& ElemTM, float Scale, const FMaterialRenderProxy* MaterialRenderProxy) const;
+	ENGINE_API FBox CalcAABB(const FTransform& BoneTM, float Scale) const;
+
+	ENGINE_API void ScaleElem(FVector DeltaSize, float MinSize);
 };
 
 /** Container for an aggregate of collision shapes */
@@ -364,19 +377,19 @@ struct ENGINE_API FKAggregateGeom
 	/** Release the RenderInfo (if its there) and safely clean up any resources. Call on the game thread. */
 	void FreeRenderInfo();
 
-	FBox CalcAABB(const FTransform& Transform);
+	FBox CalcAABB(const FTransform& Transform) const;
 
 	/**
-		* Calculates a tight box-sphere bounds for the aggregate geometry; this is more expensive than CalcAABB
-		* (tight meaning the sphere may be smaller than would be required to encompass the AABB, but all individual components lie within both the box and the sphere)
-		*
-		* @param Output The output box-sphere bounds calculated for this set of aggregate geometry
-		*	@param LocalToWorld Transform
-		*/
-	void CalcBoxSphereBounds(FBoxSphereBounds& Output, const FTransform& LocalToWorld);
+	 * Calculates a tight box-sphere bounds for the aggregate geometry; this is more expensive than CalcAABB
+	 * (tight meaning the sphere may be smaller than would be required to encompass the AABB, but all individual components lie within both the box and the sphere)
+	 *
+	 *  @param Output The output box-sphere bounds calculated for this set of aggregate geometry
+	 *  @param LocalToWorld Transform
+	 */
+	void CalcBoxSphereBounds(FBoxSphereBounds& Output, const FTransform& LocalToWorld) const;
 
 	/** Returns the volume of this element */
-	float GetVolume(const FVector& Scale) const;
+	float GetVolume(const FVector& Scale3D) const;
 };
 
 UCLASS(hidecategories=Object, MinimalAPI, dependson=BodyInstance)
@@ -484,6 +497,9 @@ public:
 	virtual void FinishDestroy() OVERRIDE;
 	virtual void PostLoad() OVERRIDE;
 	virtual void PostInitProperties() OVERRIDE;
+#if WITH_EDITOR
+	virtual void PostEditUndo() OVERRIDE;
+#endif // WITH_EDITOR
 	virtual SIZE_T GetResourceSize(EResourceSizeMode::Type Mode) OVERRIDE;
 	// End UObject interface.
 
diff --git a/Engine/Source/Runtime/Engine/Private/BoxComponent.cpp b/Engine/Source/Runtime/Engine/Private/BoxComponent.cpp
--- a/Engine/Source/Runtime/Engine/Private/BoxComponent.cpp
+++ b/Engine/Source/Runtime/Engine/Private/BoxComponent.cpp
@@ -38,7 +38,7 @@ void UBoxComponent::UpdateBodySetup()
 	{
 		ShapeBodySetup = ConstructObject<UBodySetup>(UBodySetup::StaticClass(), this);
 		ShapeBodySetup->CollisionTraceFlag = CTF_UseSimpleAsComplex;
-		ShapeBodySetup->AggGeom.BoxElems.AddZeroed();
+		ShapeBodySetup->AggGeom.BoxElems.Add(FKBoxElem());
 	}
 
 	check(ShapeBodySetup->AggGeom.BoxElems.Num() == 1);
diff --git a/Engine/Source/Runtime/Engine/Private/CapsuleComponent.cpp b/Engine/Source/Runtime/Engine/Private/CapsuleComponent.cpp
--- a/Engine/Source/Runtime/Engine/Private/CapsuleComponent.cpp
+++ b/Engine/Source/Runtime/Engine/Private/CapsuleComponent.cpp
@@ -146,8 +146,7 @@ void UCapsuleComponent::UpdateBodySetup()
 	{
 		ShapeBodySetup = ConstructObject<UBodySetup>(UBodySetup::StaticClass(), this);
 		ShapeBodySetup->CollisionTraceFlag = CTF_UseSimpleAsComplex;
-
-		ShapeBodySetup->AggGeom.SphylElems.AddZeroed();
+		ShapeBodySetup->AggGeom.SphylElems.Add(FKSphylElem());
 	}
 
 	check(ShapeBodySetup->AggGeom.SphylElems.Num() == 1);
diff --git a/Engine/Source/Runtime/Engine/Private/Collision/KAggregateGeom.cpp b/Engine/Source/Runtime/Engine/Private/Collision/KAggregateGeom.cpp
--- a/Engine/Source/Runtime/Engine/Private/Collision/KAggregateGeom.cpp
+++ b/Engine/Source/Runtime/Engine/Private/Collision/KAggregateGeom.cpp
@@ -36,7 +36,7 @@ float SelectMinScale(FVector Scale)
 	return Min;
 }
 
-FBox FKAggregateGeom::CalcAABB(const FTransform& Transform)
+FBox FKAggregateGeom::CalcAABB(const FTransform& Transform) const
 {
 	const FVector Scale3D = Transform.GetScale3D();
 	FTransform BoneTM = Transform;
@@ -81,7 +81,7 @@ FBox FKAggregateGeom::CalcAABB(const FTransform& Transform)
   * @param Output The output box-sphere bounds calculated for this set of aggregate geometry
   *	@param LocalToWorld Transform
   */
-void FKAggregateGeom::CalcBoxSphereBounds(FBoxSphereBounds& Output, const FTransform& LocalToWorld)
+void FKAggregateGeom::CalcBoxSphereBounds(FBoxSphereBounds& Output, const FTransform& LocalToWorld) const
 {
 	// Calculate the AABB
 	const FBox AABB = CalcAABB(LocalToWorld);
@@ -95,7 +95,7 @@ void FKAggregateGeom::CalcBoxSphereBounds(FBoxSphereBounds& Output, const FTrans
 		float RadiusSquared = 0.0f;
 		for (int32 i = 0; i < ConvexElems.Num(); i++)
 		{
-			FKConvexElem& Elem = ConvexElems[i];
+			const FKConvexElem& Elem = ConvexElems[i];
 			for (int32 j = 0; j < Elem.VertexData.Num(); ++j)
 			{
 				const FVector Point = LocalToWorld.TransformPosition(Elem.VertexData[j]);
@@ -249,7 +249,7 @@ static bool EnsureHullIsValid(TArray<FVector>& InVerts)
 ///////////// FKSphereElem ////////////
 ///////////////////////////////////////
 
-FBox FKSphereElem::CalcAABB(const FTransform& BoneTM, float Scale)
+FBox FKSphereElem::CalcAABB(const FTransform& BoneTM, float Scale) const
 {
 	FTransform ElemTM = GetTransform();
 	ElemTM.ScaleTranslation( FVector(Scale) );
@@ -262,13 +262,12 @@ FBox FKSphereElem::CalcAABB(const FTransform& BoneTM, float Scale)
 }
 
 
-
 ///////////////////////////////////////
 ////////////// FKBoxElem //////////////
 ///////////////////////////////////////
 
 
-FBox FKBoxElem::CalcAABB(const FTransform& BoneTM, float Scale)
+FBox FKBoxElem::CalcAABB(const FTransform& BoneTM, float Scale) const
 {
 	FTransform ElemTM = GetTransform();
 	ElemTM.ScaleTranslation( FVector(Scale) );
@@ -285,7 +284,7 @@ FBox FKBoxElem::CalcAABB(const FTransform& BoneTM, float Scale)
 ////////////// FKSphylElem ////////////
 ///////////////////////////////////////
 
-FBox FKSphylElem::CalcAABB(const FTransform& BoneTM, float Scale)
+FBox FKSphylElem::CalcAABB(const FTransform& BoneTM, float Scale) const
 {
 	FTransform ElemTM = GetTransform();
 	ElemTM.ScaleTranslation( FVector(Scale) );
@@ -321,10 +320,10 @@ void FKConvexElem::Reset()
 	ElemBox.Init();
 }
 
-FBox FKConvexElem::CalcAABB(const FTransform& BoneTM, const FVector& Scale3D)
+FBox FKConvexElem::CalcAABB(const FTransform& BoneTM, const FVector& Scale3D) const
 {
 	// Zero out rotation and location so we transform by scale along
-	const FTransform LocalToWorld = FTransform(FQuat::Identity, FVector::ZeroVector, Scale3D ) * BoneTM;
+	const FTransform LocalToWorld = FTransform(FQuat::Identity, FVector::ZeroVector, Scale3D) * BoneTM;
 
 	return ElemBox.TransformBy(LocalToWorld);
 }
diff --git a/Engine/Source/Runtime/Engine/Private/PhysicsEngine/BodySetup.cpp b/Engine/Source/Runtime/Engine/Private/PhysicsEngine/BodySetup.cpp
--- a/Engine/Source/Runtime/Engine/Private/PhysicsEngine/BodySetup.cpp
+++ b/Engine/Source/Runtime/Engine/Private/PhysicsEngine/BodySetup.cpp
@@ -38,6 +38,7 @@ UBodySetup::UBodySetup(const class FPostConstructInitializeProperties& PCIP)
 	DefaultInstance.SetObjectType(ECC_PhysicsBody);
 	BuildScale_DEPRECATED = 1.0f;
 	BuildScale3D = FVector(1.0f, 1.0f, 1.0f);
+	SetFlags(RF_Transactional);
 }
 
 void UBodySetup::CopyBodyPropertiesFrom(class UBodySetup* FromSetup)
@@ -843,6 +844,20 @@ void UBodySetup::PostInitProperties()
 	}
 }
 
+#if WITH_EDITOR
+void UBodySetup::PostEditUndo()
+{
+	Super::PostEditUndo();
+
+	// If we have any convex elems, ensure they are recreated whenever anything is modified!
+	if(AggGeom.ConvexElems.Num() > 0)
+	{
+		InvalidatePhysicsData();
+		CreatePhysicsMeshes();
+	}
+}
+#endif // WITH_EDITOR
+
 SIZE_T UBodySetup::GetResourceSize( EResourceSizeMode::Type Mode )
 {
 	SIZE_T ResourceSize = 0;
@@ -951,6 +966,13 @@ int32 FKAggregateGeom::GetElementCount(int32 Type) const
 	}
 }
 
+void FKConvexElem::ScaleElem(FVector DeltaSize, float MinSize)
+{
+	FTransform Transform = GetTransform();
+	Transform.SetScale3D(Transform.GetScale3D() + DeltaSize);
+	SetTransform(Transform);
+}
+
 // References: 
 // http://amp.ece.cmu.edu/Publication/Cha/icip01_Cha.pdf
 // http://stackoverflow.com/questions/1406029/how-to-calculate-the-volume-of-a-3d-mesh-object-the-surface-of-which-is-made-up
@@ -1007,6 +1029,18 @@ void FKSphereElem::Serialize( const FArchive& Ar )
 	}
 }
 
+void FKSphereElem::ScaleElem(FVector DeltaSize, float MinSize)
+{
+	// Find element with largest magnitude, btu preserve sign.
+	float DeltaRadius = DeltaSize.X;
+	if (FMath::Abs(DeltaSize.Y) > FMath::Abs(DeltaRadius))
+		DeltaRadius = DeltaSize.Y;
+	else if (FMath::Abs(DeltaSize.Z) > FMath::Abs(DeltaRadius))
+		DeltaRadius = DeltaSize.Z;
+
+	Radius = FMath::Max(Radius + DeltaRadius, MinSize);
+}
+
 void FKBoxElem::Serialize( const FArchive& Ar )
 {
 	if ( Ar.IsLoading() && Ar.UE4Ver() < VER_UE4_REFACTOR_PHYSICS_TRANSFORMS )
@@ -1016,6 +1050,14 @@ void FKBoxElem::Serialize( const FArchive& Ar )
 	}
 }
 
+void FKBoxElem::ScaleElem(FVector DeltaSize, float MinSize)
+{
+	// Sizes are lengths, so we double the delta to get similar increase in size.
+	X = FMath::Max(X + 2 * DeltaSize.X, MinSize);
+	Y = FMath::Max(Y + 2 * DeltaSize.Y, MinSize);
+	Z = FMath::Max(Z + 2 * DeltaSize.Z, MinSize);
+}
+
 void FKSphylElem::Serialize( const FArchive& Ar )
 {
 	if ( Ar.IsLoading() && Ar.UE4Ver() < VER_UE4_REFACTOR_PHYSICS_TRANSFORMS )
@@ -1025,6 +1067,25 @@ void FKSphylElem::Serialize( const FArchive& Ar )
 	}
 }
 
+void FKSphylElem::ScaleElem(FVector DeltaSize, float MinSize)
+{
+	float DeltaRadius = DeltaSize.X;
+	if (FMath::Abs(DeltaSize.Y) > FMath::Abs(DeltaRadius))
+	{
+		DeltaRadius = DeltaSize.Y;
+	}
+
+	float DeltaHeight = DeltaSize.Z;
+	float radius = FMath::Max(Radius + DeltaRadius, MinSize);
+	float length = Length + DeltaHeight;
+
+	length += Radius - radius;
+	length = FMath::Max(0.f, length);
+
+	Radius = radius;
+	Length = length;
+}
+
 class UPhysicalMaterial* UBodySetup::GetPhysMaterial() const
 {
 	UPhysicalMaterial* PhysMat = PhysMaterial;
diff --git a/Engine/Source/Runtime/Engine/Private/PhysicsEngine/PhysDrawing.cpp b/Engine/Source/Runtime/Engine/Private/PhysicsEngine/PhysDrawing.cpp
--- a/Engine/Source/Runtime/Engine/Private/PhysicsEngine/PhysDrawing.cpp
+++ b/Engine/Source/Runtime/Engine/Private/PhysicsEngine/PhysDrawing.cpp
@@ -26,7 +26,7 @@ static const FColor JointLockedColor(255,128,10);
 /////////////////////////////////////////////////////////////////////////////////////
 
 // NB: ElemTM is assumed to have no scaling in it!
-void FKSphereElem::DrawElemWire(FPrimitiveDrawInterface* PDI, const FTransform& ElemTM, float Scale, const FColor Color)
+void FKSphereElem::DrawElemWire(FPrimitiveDrawInterface* PDI, const FTransform& ElemTM, float Scale, const FColor Color) const
 {
 	FVector ElemCenter = ElemTM.GetLocation();
 	FVector X = ElemTM.GetScaledAxis( EAxis::X );
@@ -38,7 +38,7 @@ void FKSphereElem::DrawElemWire(FPrimitiveDrawInterface* PDI, const FTransform&
 	DrawCircle(PDI,ElemCenter, Y, Z, Color, Scale*Radius, DrawCollisionSides, SDPG_World);
 }
 
-void FKSphereElem::DrawElemSolid(class FPrimitiveDrawInterface* PDI, const FTransform& ElemTM, float Scale, const FMaterialRenderProxy* MaterialRenderProxy)
+void FKSphereElem::DrawElemSolid(class FPrimitiveDrawInterface* PDI, const FTransform& ElemTM, float Scale, const FMaterialRenderProxy* MaterialRenderProxy) const
 {
 	DrawSphere(PDI, ElemTM.GetLocation(), FVector( this->Radius * Scale ), DrawCollisionSides, DrawCollisionSides/2, MaterialRenderProxy, SDPG_World );
 }
@@ -49,7 +49,7 @@ void FKSphereElem::DrawElemSolid(class FPrimitiveDrawInterface* PDI, const FTran
 /////////////////////////////////////////////////////////////////////////////////////
 
 // NB: ElemTM is assumed to have no scaling in it!
-void FKBoxElem::DrawElemWire(FPrimitiveDrawInterface* PDI, const FTransform& ElemTM, float Scale, const FColor Color)
+void FKBoxElem::DrawElemWire(FPrimitiveDrawInterface* PDI, const FTransform& ElemTM, float Scale, const FColor Color) const
 {
 	FVector	B[2], P, Q, Radii;
 
@@ -83,7 +83,7 @@ void FKBoxElem::DrawElemWire(FPrimitiveDrawInterface* PDI, const FTransform& Ele
 	}
 }
 
-void FKBoxElem::DrawElemSolid(class FPrimitiveDrawInterface* PDI, const FTransform& ElemTM, float Scale, const FMaterialRenderProxy* MaterialRenderProxy)
+void FKBoxElem::DrawElemSolid(class FPrimitiveDrawInterface* PDI, const FTransform& ElemTM, float Scale, const FMaterialRenderProxy* MaterialRenderProxy) const
 {
 	DrawBox(PDI, ElemTM.ToMatrixWithScale(), 0.5f * FVector(X, Y, Z), MaterialRenderProxy, SDPG_World );
 }
@@ -106,7 +106,7 @@ static void DrawHalfCircle(FPrimitiveDrawInterface* PDI, const FVector& Base, co
 }
 
 // NB: ElemTM is assumed to have no scaling in it!
-void FKSphylElem::DrawElemWire(FPrimitiveDrawInterface* PDI, const FTransform& ElemTM, float Scale, const FColor Color)
+void FKSphylElem::DrawElemWire(FPrimitiveDrawInterface* PDI, const FTransform& ElemTM, float Scale, const FColor Color) const
 {
 	FVector Origin = ElemTM.GetLocation();
 	FVector XAxis = ElemTM.GetScaledAxis( EAxis::X );
@@ -136,7 +136,7 @@ void FKSphylElem::DrawElemWire(FPrimitiveDrawInterface* PDI, const FTransform& E
 	PDI->DrawLine(TopEnd - Scale*Radius*YAxis, BottomEnd - Scale*Radius*YAxis, Color, SDPG_World);
 }
 
-void FKSphylElem::DrawElemSolid(FPrimitiveDrawInterface* PDI, const FTransform& ElemTM, float Scale, const FMaterialRenderProxy* MaterialRenderProxy)
+void FKSphylElem::DrawElemSolid(FPrimitiveDrawInterface* PDI, const FTransform& ElemTM, float Scale, const FMaterialRenderProxy* MaterialRenderProxy) const
 {
 	const int32 NumSides = DrawCollisionSides;
 	const int32 NumRings = (DrawCollisionSides/2) + 1;
@@ -240,7 +240,7 @@ void FKSphylElem::DrawElemSolid(FPrimitiveDrawInterface* PDI, const FTransform&
 // FKConvexElem
 /////////////////////////////////////////////////////////////////////////////////////
 
-void FKConvexElem::DrawElemWire(FPrimitiveDrawInterface* PDI, const FTransform& ElemTM, const FColor Color)
+void FKConvexElem::DrawElemWire(FPrimitiveDrawInterface* PDI, const FTransform& ElemTM, const FColor Color) const
 {
 #if WITH_PHYSX
 
@@ -524,7 +524,6 @@ void FKAggregateGeom::DrawAggGeom(FPrimitiveDrawInterface* PDI, const FTransform
 				FTransform ElemTM = ConvexElems[i].GetTransform();
 				ElemTM *= Transform;
 				ConvexElems[i].DrawElemWire(PDI, ElemTM, ConvexColor);
-
 			}
 		}
 	}
diff --git a/Engine/Source/Runtime/Engine/Private/SphereComponent.cpp b/Engine/Source/Runtime/Engine/Private/SphereComponent.cpp
--- a/Engine/Source/Runtime/Engine/Private/SphereComponent.cpp
+++ b/Engine/Source/Runtime/Engine/Private/SphereComponent.cpp
@@ -29,7 +29,7 @@ void USphereComponent::UpdateBodySetup()
 	{
 		ShapeBodySetup = ConstructObject<UBodySetup>(UBodySetup::StaticClass(), this);
 		ShapeBodySetup->CollisionTraceFlag = CTF_UseSimpleAsComplex;
-		ShapeBodySetup->AggGeom.SphereElems.AddZeroed();
+		ShapeBodySetup->AggGeom.SphereElems.Add(FKSphereElem());
 	}
 
 	check (ShapeBodySetup->AggGeom.SphereElems.Num() == 1);
