UE4 Refactoring. Changed OVERRIDE and FINAL macros to keywords override and final.

diff --git a/Engine/Source/Editor/SceneOutliner/Public/SceneOutlinerTreeItems.h b/Engine/Source/Editor/SceneOutliner/Public/SceneOutlinerTreeItems.h
--- a/Engine/Source/Editor/SceneOutliner/Public/SceneOutlinerTreeItems.h
+++ b/Engine/Source/Editor/SceneOutliner/Public/SceneOutlinerTreeItems.h
@@ -68,10 +68,10 @@ struct SCENEOUTLINER_API TOutlinerActorTreeItem : TOutlinerTreeItem
 	}
 
 	/** Determine if this tree item is visible or not */
-	virtual bool IsVisible() const OVERRIDE;
+	virtual bool IsVisible() const override;
 
 	/** Set the visibility of this item */
-	virtual void SetIsVisible(bool bIsVisible) OVERRIDE;
+	virtual void SetIsVisible(bool bIsVisible) override;
 };
 
 struct SCENEOUTLINER_API TOutlinerFolderTreeItem : TOutlinerTreeItem
@@ -93,10 +93,10 @@ struct SCENEOUTLINER_API TOutlinerFolderTreeItem : TOutlinerTreeItem
 	}
 
 	/** Determine if this tree item is visible or not */
-	virtual bool IsVisible() const OVERRIDE;
+	virtual bool IsVisible() const override;
 	
 	/** Set the visibility of this item */
-	virtual void SetIsVisible(bool bIsVisible) OVERRIDE;
+	virtual void SetIsVisible(bool bIsVisible) override;
 
 	/** Parse a new path (including leaf-name) into this tree item. Does not do any notification */
 	void ParsePath(FName InPath)
diff --git a/Engine/Source/Editor/Sequencer/Private/SAnimationOutlinerView.h b/Engine/Source/Editor/Sequencer/Private/SAnimationOutlinerView.h
--- a/Engine/Source/Editor/Sequencer/Private/SAnimationOutlinerView.h
+++ b/Engine/Source/Editor/Sequencer/Private/SAnimationOutlinerView.h
@@ -25,8 +25,8 @@ public:
 	const TSharedPtr< const FSequencerDisplayNode> GetDisplayNode() const { return DisplayNode; }
 private:
 	/** SWidget Interface */
-	virtual FReply OnMouseButtonDown( const FGeometry& MyGeometry, const FPointerEvent& MouseEvent ) OVERRIDE;
-	virtual FReply OnMouseButtonUp( const FGeometry& MyGeometry, const FPointerEvent& MouseEvent ) OVERRIDE;
+	virtual FReply OnMouseButtonDown( const FGeometry& MyGeometry, const FPointerEvent& MouseEvent ) override;
+	virtual FReply OnMouseButtonUp( const FGeometry& MyGeometry, const FPointerEvent& MouseEvent ) override;
 
 	/**
 	 * Called when the expander arrow is clicked
@@ -92,7 +92,7 @@ private:
 	void GenerateWidgetForNode( TSharedRef<FSequencerDisplayNode>& InLayoutNode );
 
 	/** SPanel Interface */
-	virtual void OnArrangeChildren( const FGeometry& AllottedGeometry, FArrangedChildren& ArrangedChildren ) const OVERRIDE;
+	virtual void OnArrangeChildren( const FGeometry& AllottedGeometry, FArrangedChildren& ArrangedChildren ) const override;
 
 	/**
 	 * Called when a node is selected in the outliner
diff --git a/Engine/Source/Editor/Sequencer/Private/SSection.h b/Engine/Source/Editor/Sequencer/Private/SSection.h
--- a/Engine/Source/Editor/Sequencer/Private/SSection.h
+++ b/Engine/Source/Editor/Sequencer/Private/SSection.h
@@ -73,14 +73,14 @@ private:
 	void CheckForEdgeInteraction( const FPointerEvent& MousePosition, const FGeometry& SectionGeometry );
 
 	/** SWidget interface */
-	virtual int32 OnPaint( const FGeometry& AllottedGeometry, const FSlateRect& MyClippingRect, FSlateWindowElementList& OutDrawElements, int32 LayerId, const FWidgetStyle& InWidgetStyle, bool bParentEnabled ) const OVERRIDE;
-	virtual void Tick( const FGeometry& AllottedGeometry, const double InCurrentTime, const float InDeltaTime ) OVERRIDE;
-	virtual FReply OnMouseButtonDown( const FGeometry& MyGeometry, const FPointerEvent& MouseEvent ) OVERRIDE;
-	virtual FReply OnMouseButtonUp( const FGeometry& MyGeometry, const FPointerEvent& MouseEvent ) OVERRIDE;
-	virtual FReply OnMouseButtonDoubleClick( const FGeometry& MyGeometry, const FPointerEvent& MouseEvent ) OVERRIDE;
-	virtual FReply OnMouseMove( const FGeometry& MyGeometry, const FPointerEvent& MouseEvent ) OVERRIDE;
-	virtual void OnMouseLeave( const FPointerEvent& MouseEvent ) OVERRIDE;
-	virtual FCursorReply OnCursorQuery( const FGeometry& MyGeometry, const FPointerEvent& CursorEvent ) const OVERRIDE;
+	virtual int32 OnPaint( const FGeometry& AllottedGeometry, const FSlateRect& MyClippingRect, FSlateWindowElementList& OutDrawElements, int32 LayerId, const FWidgetStyle& InWidgetStyle, bool bParentEnabled ) const override;
+	virtual void Tick( const FGeometry& AllottedGeometry, const double InCurrentTime, const float InDeltaTime ) override;
+	virtual FReply OnMouseButtonDown( const FGeometry& MyGeometry, const FPointerEvent& MouseEvent ) override;
+	virtual FReply OnMouseButtonUp( const FGeometry& MyGeometry, const FPointerEvent& MouseEvent ) override;
+	virtual FReply OnMouseButtonDoubleClick( const FGeometry& MyGeometry, const FPointerEvent& MouseEvent ) override;
+	virtual FReply OnMouseMove( const FGeometry& MyGeometry, const FPointerEvent& MouseEvent ) override;
+	virtual void OnMouseLeave( const FPointerEvent& MouseEvent ) override;
+	virtual FCursorReply OnCursorQuery( const FGeometry& MyGeometry, const FPointerEvent& CursorEvent ) const override;
 
 	/**
 	 * Paints keys visible inside the section
diff --git a/Engine/Source/Editor/Sequencer/Private/SSequencer.cpp b/Engine/Source/Editor/Sequencer/Private/SSequencer.cpp
--- a/Engine/Source/Editor/Sequencer/Private/SSequencer.cpp
+++ b/Engine/Source/Editor/Sequencer/Private/SSequencer.cpp
@@ -60,7 +60,7 @@ public:
 		}
 	}
 	
-	virtual int32 OnPaint( const FGeometry& AllottedGeometry, const FSlateRect& MyClippingRect, FSlateWindowElementList& OutDrawElements, int32 LayerId, const FWidgetStyle& InWidgetStyle, bool bParentEnabled ) const OVERRIDE
+	virtual int32 OnPaint( const FGeometry& AllottedGeometry, const FSlateRect& MyClippingRect, FSlateWindowElementList& OutDrawElements, int32 LayerId, const FWidgetStyle& InWidgetStyle, bool bParentEnabled ) const override
 	{
 		float Alpha = Sequencer.Pin()->GetOverlayFadeCurve();
 		
@@ -92,7 +92,7 @@ public:
 		return LayerId;
 	}
 
-	virtual FVector2D ComputeDesiredSize() const OVERRIDE
+	virtual FVector2D ComputeDesiredSize() const override
 	{
 		return FVector2D(100, 100);
 	}
diff --git a/Engine/Source/Editor/Sequencer/Private/SSequencer.h b/Engine/Source/Editor/Sequencer/Private/SSequencer.h
--- a/Engine/Source/Editor/Sequencer/Private/SSequencer.h
+++ b/Engine/Source/Editor/Sequencer/Private/SSequencer.h
@@ -76,8 +76,8 @@ public:
 
 	~SSequencer();
 
-	virtual void Tick( const FGeometry& AllottedGeometry, const double InCurrentTime, const float InDeltaTime ) OVERRIDE;
-	virtual bool SupportsKeyboardFocus() const OVERRIDE { return true; }
+	virtual void Tick( const FGeometry& AllottedGeometry, const double InCurrentTime, const float InDeltaTime ) override;
+	virtual bool SupportsKeyboardFocus() const override { return true; }
 
 	/** Updates the layout node tree from movie scene data */
 	void UpdateLayoutTree();
@@ -144,11 +144,11 @@ private:
 
 	/** SWidget interface */
 	/** @todo Sequencer Basic drag and drop support.  Doesn't belong here most likely */
-	virtual void OnDragEnter( const FGeometry& MyGeometry, const FDragDropEvent& DragDropEvent ) OVERRIDE;
-	virtual void OnDragLeave( const FDragDropEvent& DragDropEvent ) OVERRIDE;
-	virtual FReply OnDragOver( const FGeometry& MyGeometry, const FDragDropEvent& DragDropEvent ) OVERRIDE;
-	virtual FReply OnDrop( const FGeometry& MyGeometry, const FDragDropEvent& DragDropEvent ) OVERRIDE;
-	virtual FReply OnKeyDown( const FGeometry& MyGeometry, const FKeyboardEvent& InKeyboardEvent ) OVERRIDE;
+	virtual void OnDragEnter( const FGeometry& MyGeometry, const FDragDropEvent& DragDropEvent ) override;
+	virtual void OnDragLeave( const FDragDropEvent& DragDropEvent ) override;
+	virtual FReply OnDragOver( const FGeometry& MyGeometry, const FDragDropEvent& DragDropEvent ) override;
+	virtual FReply OnDrop( const FGeometry& MyGeometry, const FDragDropEvent& DragDropEvent ) override;
+	virtual FReply OnKeyDown( const FGeometry& MyGeometry, const FKeyboardEvent& InKeyboardEvent ) override;
 
 	/**
 	 * Called when one or more assets are dropped into the widget
diff --git a/Engine/Source/Editor/Sequencer/Private/SSequencerSectionAreaView.h b/Engine/Source/Editor/Sequencer/Private/SSequencerSectionAreaView.h
--- a/Engine/Source/Editor/Sequencer/Private/SSequencerSectionAreaView.h
+++ b/Engine/Source/Editor/Sequencer/Private/SSequencerSectionAreaView.h
@@ -35,11 +35,11 @@ public:
 
 private:
 	/** SWidget Interface */
-	virtual int32 OnPaint( const FGeometry& AllottedGeometry, const FSlateRect& MyClippingRect, FSlateWindowElementList& OutDrawElements, int32 LayerId, const FWidgetStyle& InWidgetStyle, bool bParentEnabled ) const OVERRIDE;
-	virtual FReply OnMouseButtonDown( const FGeometry& MyGeometry, const FPointerEvent& MouseEvent ) OVERRIDE;
+	virtual int32 OnPaint( const FGeometry& AllottedGeometry, const FSlateRect& MyClippingRect, FSlateWindowElementList& OutDrawElements, int32 LayerId, const FWidgetStyle& InWidgetStyle, bool bParentEnabled ) const override;
+	virtual FReply OnMouseButtonDown( const FGeometry& MyGeometry, const FPointerEvent& MouseEvent ) override;
 
 	/** SPanel Interface */
-	virtual void OnArrangeChildren( const FGeometry& AllottedGeometry, FArrangedChildren& ArrangedChildren ) const OVERRIDE;
+	virtual void OnArrangeChildren( const FGeometry& AllottedGeometry, FArrangedChildren& ArrangedChildren ) const override;
 
 	/**
 	 * Returns the visibility of a section
diff --git a/Engine/Source/Editor/Sequencer/Private/SSequencerSectionOverlay.h b/Engine/Source/Editor/Sequencer/Private/SSequencerSectionOverlay.h
--- a/Engine/Source/Editor/Sequencer/Private/SSequencerSectionOverlay.h
+++ b/Engine/Source/Editor/Sequencer/Private/SSequencerSectionOverlay.h
@@ -26,11 +26,11 @@ public:
 
 private:
 	/** SWidget Interface */
-	virtual int32 OnPaint( const FGeometry& AllottedGeometry, const FSlateRect& MyClippingRect, FSlateWindowElementList& OutDrawElements, int32 LayerId, const FWidgetStyle& InWidgetStyle, bool bParentEnabled ) const OVERRIDE;
-	virtual FReply OnMouseButtonDown( const FGeometry& MyGeometry, const FPointerEvent& MouseEvent ) OVERRIDE;
-	virtual FReply OnMouseButtonUp( const FGeometry& MyGeometry, const FPointerEvent& MouseEvent ) OVERRIDE;
-	virtual FReply OnMouseMove( const FGeometry& MyGeometry, const FPointerEvent& MouseEvent ) OVERRIDE;
-	virtual FReply OnMouseWheel( const FGeometry& MyGeometry, const FPointerEvent& MouseEvent ) OVERRIDE;
+	virtual int32 OnPaint( const FGeometry& AllottedGeometry, const FSlateRect& MyClippingRect, FSlateWindowElementList& OutDrawElements, int32 LayerId, const FWidgetStyle& InWidgetStyle, bool bParentEnabled ) const override;
+	virtual FReply OnMouseButtonDown( const FGeometry& MyGeometry, const FPointerEvent& MouseEvent ) override;
+	virtual FReply OnMouseButtonUp( const FGeometry& MyGeometry, const FPointerEvent& MouseEvent ) override;
+	virtual FReply OnMouseMove( const FGeometry& MyGeometry, const FPointerEvent& MouseEvent ) override;
+	virtual FReply OnMouseWheel( const FGeometry& MyGeometry, const FPointerEvent& MouseEvent ) override;
 
 private:
 	/** Controller for manipulating time */
diff --git a/Engine/Source/Editor/Sequencer/Private/SSequencerTrackViewPanel.h b/Engine/Source/Editor/Sequencer/Private/SSequencerTrackViewPanel.h
--- a/Engine/Source/Editor/Sequencer/Private/SSequencerTrackViewPanel.h
+++ b/Engine/Source/Editor/Sequencer/Private/SSequencerTrackViewPanel.h
@@ -19,7 +19,7 @@ public:
 	}
 
 	/** SWidget interface */
-	virtual FVector2D ComputeDesiredSize() const OVERRIDE
+	virtual FVector2D ComputeDesiredSize() const override
 	{
 		float NodeHeight = RootNode->GetNodeHeight();
 	
@@ -33,7 +33,7 @@ public:
 		return FVector2D( 100.0f, NodeHeight );
 	}
 
-	virtual FChildren* GetChildren() OVERRIDE
+	virtual FChildren* GetChildren() override
 	{ 
 		return &Children; 
 	}
diff --git a/Engine/Source/Editor/Sequencer/Private/SectionDragOperations.h b/Engine/Source/Editor/Sequencer/Private/SectionDragOperations.h
--- a/Engine/Source/Editor/Sequencer/Private/SectionDragOperations.h
+++ b/Engine/Source/Editor/Sequencer/Private/SectionDragOperations.h
@@ -85,10 +85,10 @@ public:
 	{}
 
 	/** FSequencerDragOperation interface */
-	virtual void OnBeginDrag(const FVector2D& LocalMousePos, TSharedPtr<FTrackNode> SequencerNode) OVERRIDE;
-	virtual void OnEndDrag(TSharedPtr<FTrackNode> SequencerNode) OVERRIDE;
-	virtual void OnDrag( const FPointerEvent& MouseEvent, const FVector2D& LocalMousePos, const FTimeToPixel& TimeToPixelConverter, TSharedPtr<FTrackNode> SequencerNode ) OVERRIDE;
-	virtual FCursorReply GetCursor() const OVERRIDE { return FCursorReply::Cursor( EMouseCursor::ResizeLeftRight ); }
+	virtual void OnBeginDrag(const FVector2D& LocalMousePos, TSharedPtr<FTrackNode> SequencerNode) override;
+	virtual void OnEndDrag(TSharedPtr<FTrackNode> SequencerNode) override;
+	virtual void OnDrag( const FPointerEvent& MouseEvent, const FVector2D& LocalMousePos, const FTimeToPixel& TimeToPixelConverter, TSharedPtr<FTrackNode> SequencerNode ) override;
+	virtual FCursorReply GetCursor() const override { return FCursorReply::Cursor( EMouseCursor::ResizeLeftRight ); }
 private:
 	/** The section we are interacting with */
 	TWeakObjectPtr<UMovieSceneSection> Section;
@@ -108,10 +108,10 @@ public:
 	{}
 
 	/** FSequencerDragOperation interface */
-	virtual void OnBeginDrag(const FVector2D& LocalMousePos, TSharedPtr<FTrackNode> SequencerNode) OVERRIDE;
-	virtual void OnEndDrag(TSharedPtr<FTrackNode> SequencerNode) OVERRIDE;
-	virtual void OnDrag( const FPointerEvent& MouseEvent, const FVector2D& LocalMousePos, const FTimeToPixel& TimeToPixelConverter, TSharedPtr<FTrackNode> SequencerNode ) OVERRIDE;
-	virtual FCursorReply GetCursor() const OVERRIDE { return FCursorReply::Cursor( EMouseCursor::CardinalCross ); }
+	virtual void OnBeginDrag(const FVector2D& LocalMousePos, TSharedPtr<FTrackNode> SequencerNode) override;
+	virtual void OnEndDrag(TSharedPtr<FTrackNode> SequencerNode) override;
+	virtual void OnDrag( const FPointerEvent& MouseEvent, const FVector2D& LocalMousePos, const FTimeToPixel& TimeToPixelConverter, TSharedPtr<FTrackNode> SequencerNode ) override;
+	virtual FCursorReply GetCursor() const override { return FCursorReply::Cursor( EMouseCursor::CardinalCross ); }
 
 private:
 	/** Determines if the section, when moved to this the track index by a delta time, would still be valid */
@@ -136,9 +136,9 @@ public:
 	{}
 
 	/** FSequencerDragOperation interface */
-	virtual void OnBeginDrag(const FVector2D& LocalMousePos, TSharedPtr<FTrackNode> SequencerNode) OVERRIDE;
-	virtual void OnEndDrag(TSharedPtr<FTrackNode> SequencerNode) OVERRIDE;
-	virtual void OnDrag( const FPointerEvent& MouseEvent, const FVector2D& LocalMousePos, const FTimeToPixel& TimeToPixelConverter, TSharedPtr<FTrackNode> SequencerNode ) OVERRIDE;
+	virtual void OnBeginDrag(const FVector2D& LocalMousePos, TSharedPtr<FTrackNode> SequencerNode) override;
+	virtual void OnEndDrag(TSharedPtr<FTrackNode> SequencerNode) override;
+	virtual void OnDrag( const FPointerEvent& MouseEvent, const FVector2D& LocalMousePos, const FTimeToPixel& TimeToPixelConverter, TSharedPtr<FTrackNode> SequencerNode ) override;
 
 protected:	
 	/** Gets all potential snap times for this drag operation */
diff --git a/Engine/Source/Editor/Sequencer/Private/SectionLayoutBuilder.h b/Engine/Source/Editor/Sequencer/Private/SectionLayoutBuilder.h
--- a/Engine/Source/Editor/Sequencer/Private/SectionLayoutBuilder.h
+++ b/Engine/Source/Editor/Sequencer/Private/SectionLayoutBuilder.h
@@ -11,10 +11,10 @@ public:
 	FSectionLayoutBuilder( TSharedRef<FTrackNode> InRootNode );
 
 	/** ISectionLayoutBuilder Interface */
-	virtual void PushCategory( FName CategoryName, const FString& DisplayLabel ) OVERRIDE;
-	virtual void SetSectionAsKeyArea( TSharedRef<IKeyArea> KeyArea ) OVERRIDE;
-	virtual void AddKeyArea( FName KeyAreaName, const FString& DisplayName, TSharedRef<IKeyArea> KeyArea ) OVERRIDE;
-	virtual void PopCategory() OVERRIDE;
+	virtual void PushCategory( FName CategoryName, const FString& DisplayLabel ) override;
+	virtual void SetSectionAsKeyArea( TSharedRef<IKeyArea> KeyArea ) override;
+	virtual void AddKeyArea( FName KeyAreaName, const FString& DisplayName, TSharedRef<IKeyArea> KeyArea ) override;
+	virtual void PopCategory() override;
 private:
 	/** Root node of the tree */
 	TSharedRef<FTrackNode> RootNode;
diff --git a/Engine/Source/Editor/Sequencer/Private/Sequencer.h b/Engine/Source/Editor/Sequencer/Private/Sequencer.h
--- a/Engine/Source/Editor/Sequencer/Private/Sequencer.h
+++ b/Engine/Source/Editor/Sequencer/Private/Sequencer.h
@@ -32,31 +32,31 @@ public:
 	virtual ~FSequencer();
 
 	/** FGCObject interface */
-	virtual void AddReferencedObjects(FReferenceCollector& Collector) OVERRIDE;
+	virtual void AddReferencedObjects(FReferenceCollector& Collector) override;
 
 
 	/** ISequencer interface */
-	virtual TSharedRef<SWidget> GetSequencerWidget() const OVERRIDE { return SequencerWidget.ToSharedRef(); }
+	virtual TSharedRef<SWidget> GetSequencerWidget() const override { return SequencerWidget.ToSharedRef(); }
 	virtual UMovieScene* GetRootMovieScene() const;
-	virtual UMovieScene* GetFocusedMovieScene() const OVERRIDE;
-	virtual TSharedRef<FMovieSceneInstance> GetRootMovieSceneInstance() const OVERRIDE;
-	virtual TSharedRef<FMovieSceneInstance> GetFocusedMovieSceneInstance() const OVERRIDE;
-	virtual void FocusSubMovieScene( TSharedRef<FMovieSceneInstance> SubMovieSceneInstance ) OVERRIDE;
-	TSharedRef<FMovieSceneInstance> GetInstanceForSubMovieSceneSection( UMovieSceneSection& SubMovieSceneSection ) const OVERRIDE;
-	virtual void AddNewShot(FGuid CameraGuid) OVERRIDE;
-	virtual void AddAnimation(FGuid ObjectGuid, class UAnimSequence* AnimSequence) OVERRIDE;
-	virtual bool IsAutoKeyEnabled() const OVERRIDE;
-	virtual bool IsRecordingLive() const OVERRIDE;
-	virtual float GetCurrentLocalTime(UMovieScene& MovieScene) OVERRIDE;
-	virtual float GetGlobalTime() OVERRIDE;
-	virtual void SetGlobalTime(float Time) OVERRIDE;
-	virtual void SetPerspectiveViewportPossessionEnabled(bool bEnabled) OVERRIDE;
-	virtual FGuid GetHandleToObject(UObject* Object) OVERRIDE;
-	virtual ISequencerObjectChangeListener& GetObjectChangeListener() const OVERRIDE;
-	virtual void NotifyMovieSceneDataChanged() OVERRIDE;
-	virtual void AddSubMovieScene(UMovieScene* SubMovieScene) OVERRIDE;
-	virtual void FilterToShotSections(const TArray< TWeakObjectPtr<class UMovieSceneSection> >& ShotSections, bool bZoomToShotBounds = true) OVERRIDE;
-	virtual void FilterToSelectedShotSections(bool bZoomToShotBounds = true) OVERRIDE;
+	virtual UMovieScene* GetFocusedMovieScene() const override;
+	virtual TSharedRef<FMovieSceneInstance> GetRootMovieSceneInstance() const override;
+	virtual TSharedRef<FMovieSceneInstance> GetFocusedMovieSceneInstance() const override;
+	virtual void FocusSubMovieScene( TSharedRef<FMovieSceneInstance> SubMovieSceneInstance ) override;
+	TSharedRef<FMovieSceneInstance> GetInstanceForSubMovieSceneSection( UMovieSceneSection& SubMovieSceneSection ) const override;
+	virtual void AddNewShot(FGuid CameraGuid) override;
+	virtual void AddAnimation(FGuid ObjectGuid, class UAnimSequence* AnimSequence) override;
+	virtual bool IsAutoKeyEnabled() const override;
+	virtual bool IsRecordingLive() const override;
+	virtual float GetCurrentLocalTime(UMovieScene& MovieScene) override;
+	virtual float GetGlobalTime() override;
+	virtual void SetGlobalTime(float Time) override;
+	virtual void SetPerspectiveViewportPossessionEnabled(bool bEnabled) override;
+	virtual FGuid GetHandleToObject(UObject* Object) override;
+	virtual ISequencerObjectChangeListener& GetObjectChangeListener() const override;
+	virtual void NotifyMovieSceneDataChanged() override;
+	virtual void AddSubMovieScene(UMovieScene* SubMovieScene) override;
+	virtual void FilterToShotSections(const TArray< TWeakObjectPtr<class UMovieSceneSection> >& ShotSections, bool bZoomToShotBounds = true) override;
+	virtual void FilterToSelectedShotSections(bool bZoomToShotBounds = true) override;
 
 	bool IsPerspectiveViewportPosessionEnabled() const { return bPerspectiveViewportPossessionEnabled; }
 
@@ -233,10 +233,10 @@ public:
 
 	/** IMovieScenePlayer interface */
 	virtual void GetRuntimeObjects( TSharedRef<FMovieSceneInstance> MovieSceneInstance, const FGuid& ObjectHandle, TArray< UObject* >& OutObjects ) const;
-	virtual void UpdateViewports(AActor* ActorToViewThrough) const OVERRIDE;
-	virtual EMovieScenePlayerStatus::Type GetPlaybackStatus() const OVERRIDE;
-	virtual void AddMovieSceneInstance( class UMovieSceneSection& MovieSceneSection, TSharedRef<FMovieSceneInstance> InstanceToAdd ) OVERRIDE;
-	virtual void RemoveMovieSceneInstance( class UMovieSceneSection& MovieSceneSection, TSharedRef<FMovieSceneInstance> InstanceToRemove ) OVERRIDE;
+	virtual void UpdateViewports(AActor* ActorToViewThrough) const override;
+	virtual EMovieScenePlayerStatus::Type GetPlaybackStatus() const override;
+	virtual void AddMovieSceneInstance( class UMovieSceneSection& MovieSceneSection, TSharedRef<FMovieSceneInstance> InstanceToAdd ) override;
+	virtual void RemoveMovieSceneInstance( class UMovieSceneSection& MovieSceneSection, TSharedRef<FMovieSceneInstance> InstanceToRemove ) override;
 
 	virtual void SpawnActorsForMovie( TSharedRef<FMovieSceneInstance> MovieSceneInstance );
 
@@ -384,8 +384,8 @@ protected:
 	void BindSequencerCommands();
 
 	// Begin FEditorUndoClient Interface
-	virtual void PostUndo(bool bSuccess) OVERRIDE;
-	virtual void PostRedo(bool bSuccess) OVERRIDE { PostUndo(bSuccess); }
+	virtual void PostUndo(bool bSuccess) override;
+	virtual void PostRedo(bool bSuccess) override { PostUndo(bSuccess); }
 	// End of FEditorUndoClient
 
 private:
diff --git a/Engine/Source/Editor/Sequencer/Private/SequencerAssetEditor.h b/Engine/Source/Editor/Sequencer/Private/SequencerAssetEditor.h
--- a/Engine/Source/Editor/Sequencer/Private/SequencerAssetEditor.h
+++ b/Engine/Source/Editor/Sequencer/Private/SequencerAssetEditor.h
@@ -24,16 +24,16 @@ public:
 	 */
 	void InitSequencerAssetEditor( const EToolkitMode::Type Mode, const TSharedPtr< class IToolkitHost >& InitToolkitHost, UMovieScene* InRootMovieScene, const TArray<FOnCreateTrackEditor>& TrackEditorDelegates, bool bEditWithinLevelEditor );
 
-	TSharedRef<ISequencer> GetSequencerInterface() const OVERRIDE { return Sequencer.ToSharedRef(); }
+	TSharedRef<ISequencer> GetSequencerInterface() const override { return Sequencer.ToSharedRef(); }
 
 	/** IToolkit interface */
-	virtual FName GetToolkitFName() const OVERRIDE;
-	virtual FText GetBaseToolkitName() const OVERRIDE;
-	virtual FString GetWorldCentricTabPrefix() const OVERRIDE;
-	virtual FLinearColor GetWorldCentricTabColorScale() const OVERRIDE;
+	virtual FName GetToolkitFName() const override;
+	virtual FText GetBaseToolkitName() const override;
+	virtual FString GetWorldCentricTabPrefix() const override;
+	virtual FLinearColor GetWorldCentricTabColorScale() const override;
 
-	virtual void RegisterTabSpawners(const TSharedRef<class FTabManager>& TabManager) OVERRIDE;
-	virtual void UnregisterTabSpawners(const TSharedRef<class FTabManager>& TabManager) OVERRIDE;
+	virtual void RegisterTabSpawners(const TSharedRef<class FTabManager>& TabManager) override;
+	virtual void UnregisterTabSpawners(const TSharedRef<class FTabManager>& TabManager) override;
 
 	void UpdateViewports(AActor* ActorToViewThrough) const;
 
diff --git a/Engine/Source/Editor/Sequencer/Private/SequencerDisplayNode.cpp b/Engine/Source/Editor/Sequencer/Private/SequencerDisplayNode.cpp
--- a/Engine/Source/Editor/Sequencer/Private/SequencerDisplayNode.cpp
+++ b/Engine/Source/Editor/Sequencer/Private/SequencerDisplayNode.cpp
@@ -23,8 +23,8 @@ public:
 	SLATE_END_ARGS()
 
 	/** SLeafWidget Interface */
-	virtual int32 OnPaint( const FGeometry& AllottedGeometry, const FSlateRect& MyClippingRect, FSlateWindowElementList& OutDrawElements, int32 LayerId, const FWidgetStyle& InWidgetStyle, bool bParentEnabled ) const OVERRIDE;
-	virtual FVector2D ComputeDesiredSize() const OVERRIDE;
+	virtual int32 OnPaint( const FGeometry& AllottedGeometry, const FSlateRect& MyClippingRect, FSlateWindowElementList& OutDrawElements, int32 LayerId, const FWidgetStyle& InWidgetStyle, bool bParentEnabled ) const override;
+	virtual FVector2D ComputeDesiredSize() const override;
 	
 	void Construct( const FArguments& InArgs, TSharedRef<FSequencerDisplayNode> InRootNode )
 	{
diff --git a/Engine/Source/Editor/Sequencer/Private/SequencerDisplayNode.h b/Engine/Source/Editor/Sequencer/Private/SequencerDisplayNode.h
--- a/Engine/Source/Editor/Sequencer/Private/SequencerDisplayNode.h
+++ b/Engine/Source/Editor/Sequencer/Private/SequencerDisplayNode.h
@@ -247,10 +247,10 @@ public:
 	{}
 
 	/** FSequencerDisplayNodeInterface */
-	virtual ESequencerNode::Type GetType() const OVERRIDE { return ESequencerNode::KeyArea; }
-	virtual float GetNodeHeight() const OVERRIDE;
-	virtual FString GetDisplayName() const OVERRIDE { return DisplayName; }
-	virtual bool GetShotFilteredVisibilityToCache() const OVERRIDE;
+	virtual ESequencerNode::Type GetType() const override { return ESequencerNode::KeyArea; }
+	virtual float GetNodeHeight() const override;
+	virtual FString GetDisplayName() const override { return DisplayName; }
+	virtual bool GetShotFilteredVisibilityToCache() const override;
 
 	/**
 	 * Adds a key area to this node
@@ -303,12 +303,12 @@ public:
 	FTrackNode( FName NodeName, UMovieSceneTrack& InAssociatedType, TSharedPtr<FSequencerDisplayNode> InParentNode, FSequencerNodeTree& InParentTree );
 
 	/** FSequencerDisplayNodeInterface */
-	virtual ESequencerNode::Type GetType() const OVERRIDE { return ESequencerNode::Track; }
-	virtual float GetNodeHeight() const OVERRIDE;
-	virtual FString GetDisplayName() const OVERRIDE;
-	virtual bool IsSelectable() const OVERRIDE { return true; }
-	virtual bool GetShotFilteredVisibilityToCache() const OVERRIDE;
-	virtual void GetChildKeyAreaNodesRecursively(TArray< TSharedRef<class FSectionKeyAreaNode> >& OutNodes) const OVERRIDE;
+	virtual ESequencerNode::Type GetType() const override { return ESequencerNode::Track; }
+	virtual float GetNodeHeight() const override;
+	virtual FString GetDisplayName() const override;
+	virtual bool IsSelectable() const override { return true; }
+	virtual bool GetShotFilteredVisibilityToCache() const override;
+	virtual void GetChildKeyAreaNodesRecursively(TArray< TSharedRef<class FSectionKeyAreaNode> >& OutNodes) const override;
 
 	/**
 	 * Adds a section to this node
@@ -376,11 +376,11 @@ public:
 	{}
 
 	/** FSequencerDisplayNodeInterface */
-	virtual ESequencerNode::Type GetType() const OVERRIDE { return ESequencerNode::Object; }
-	virtual FString GetDisplayName() const OVERRIDE { return DisplayName; }
-	virtual float GetNodeHeight() const OVERRIDE;
-	virtual bool IsSelectable() const OVERRIDE { return true; }
-	virtual bool GetShotFilteredVisibilityToCache() const OVERRIDE;
+	virtual ESequencerNode::Type GetType() const override { return ESequencerNode::Object; }
+	virtual FString GetDisplayName() const override { return DisplayName; }
+	virtual float GetNodeHeight() const override;
+	virtual bool IsSelectable() const override { return true; }
+	virtual bool GetShotFilteredVisibilityToCache() const override;
 	
 	/** @return The object binding on this node */
 	const FGuid& GetObjectBinding() const { return ObjectBinding; }
@@ -415,10 +415,10 @@ public:
 	{}
 
 	/** FSequencerDisplayNodeInterface */
-	virtual ESequencerNode::Type GetType() const OVERRIDE { return ESequencerNode::Category; }
-	virtual float GetNodeHeight() const OVERRIDE;
-	virtual FString GetDisplayName() const OVERRIDE { return DisplayName; }
-	virtual bool GetShotFilteredVisibilityToCache() const OVERRIDE;
+	virtual ESequencerNode::Type GetType() const override { return ESequencerNode::Category; }
+	virtual float GetNodeHeight() const override;
+	virtual FString GetDisplayName() const override { return DisplayName; }
+	virtual bool GetShotFilteredVisibilityToCache() const override;
 private:
 	/** The display name of the category */
 	FString DisplayName;
diff --git a/Engine/Source/Editor/Sequencer/Private/SequencerModule.cpp b/Engine/Source/Editor/Sequencer/Private/SequencerModule.cpp
--- a/Engine/Source/Editor/Sequencer/Private/SequencerModule.cpp
+++ b/Engine/Source/Editor/Sequencer/Private/SequencerModule.cpp
@@ -14,36 +14,36 @@ class FSequencerModule : public ISequencerModule
 {
 
 	/** ISequencerModule interface */
-	virtual TSharedPtr<ISequencer> CreateSequencer( UMovieScene* RootMovieScene ) OVERRIDE
+	virtual TSharedPtr<ISequencer> CreateSequencer( UMovieScene* RootMovieScene ) override
 	{
 		TSharedRef< FSequencer > Sequencer = MakeShareable(new FSequencer);
 		Sequencer->InitSequencer( RootMovieScene, nullptr, TrackEditorDelegates, false );
 		return Sequencer;
 	}
 	
-	virtual TSharedPtr<ISequencer> CreateSequencerAssetEditor( const EToolkitMode::Type Mode, const TSharedPtr< class IToolkitHost >& InitToolkitHost, UMovieScene* InRootMovieScene, bool bEditWithinLevelEditor ) OVERRIDE
+	virtual TSharedPtr<ISequencer> CreateSequencerAssetEditor( const EToolkitMode::Type Mode, const TSharedPtr< class IToolkitHost >& InitToolkitHost, UMovieScene* InRootMovieScene, bool bEditWithinLevelEditor ) override
 	{
 		TSharedRef< FSequencerAssetEditor > SequencerAssetEditor = MakeShareable(new FSequencerAssetEditor);
 		SequencerAssetEditor->InitSequencerAssetEditor( Mode, InitToolkitHost, InRootMovieScene, TrackEditorDelegates, bEditWithinLevelEditor );
 		return SequencerAssetEditor->GetSequencerInterface();
 	}
 
-	virtual void RegisterTrackEditor( FOnCreateTrackEditor InOnCreateTrackEditor ) OVERRIDE
+	virtual void RegisterTrackEditor( FOnCreateTrackEditor InOnCreateTrackEditor ) override
 	{
 		TrackEditorDelegates.AddUnique( InOnCreateTrackEditor );
 	}
 
-	virtual void UnRegisterTrackEditor( FOnCreateTrackEditor InOnCreateTrackEditor ) OVERRIDE
+	virtual void UnRegisterTrackEditor( FOnCreateTrackEditor InOnCreateTrackEditor ) override
 	{
 		TrackEditorDelegates.Remove( InOnCreateTrackEditor );
 	}
 
-	virtual void StartupModule() OVERRIDE
+	virtual void StartupModule() override
 	{
 		FSequencerCommands::Register();
 	}
 
-	virtual void ShutdownModule() OVERRIDE
+	virtual void ShutdownModule() override
 	{
 		FSequencerCommands::Unregister();
 	}
diff --git a/Engine/Source/Editor/Sequencer/Private/SequencerObjectChangeListener.h b/Engine/Source/Editor/Sequencer/Private/SequencerObjectChangeListener.h
--- a/Engine/Source/Editor/Sequencer/Private/SequencerObjectChangeListener.h
+++ b/Engine/Source/Editor/Sequencer/Private/SequencerObjectChangeListener.h
@@ -14,9 +14,9 @@ public:
 	~FSequencerObjectChangeListener();
 
 	/** ISequencerObjectChangeListener interface */
-	virtual FOnAnimatablePropertyChanged& GetOnAnimatablePropertyChanged( TSubclassOf<UProperty> PropertyClass ) OVERRIDE;
-	virtual FOnPropagateObjectChanges& GetOnPropagateObjectChanges() OVERRIDE;
-	virtual void TriggerAllPropertiesChanged(UObject* Object) OVERRIDE;
+	virtual FOnAnimatablePropertyChanged& GetOnAnimatablePropertyChanged( TSubclassOf<UProperty> PropertyClass ) override;
+	virtual FOnPropagateObjectChanges& GetOnPropagateObjectChanges() override;
+	virtual void TriggerAllPropertiesChanged(UObject* Object) override;
 
 private:
 	/**
diff --git a/Engine/Source/Editor/Sequencer/Private/SequencerObjectSpawner.h b/Engine/Source/Editor/Sequencer/Private/SequencerObjectSpawner.h
--- a/Engine/Source/Editor/Sequencer/Private/SequencerObjectSpawner.h
+++ b/Engine/Source/Editor/Sequencer/Private/SequencerObjectSpawner.h
@@ -35,7 +35,7 @@ struct FPuppetActorInfo : public FBasePuppetInfo
 public:
 
 	/** @return Gets the type of puppet */
-	virtual EPuppetObjectType::Type GetType() const OVERRIDE
+	virtual EPuppetObjectType::Type GetType() const override
 	{
 		return EPuppetObjectType::Actor;
 	}
@@ -117,10 +117,10 @@ public:
 
 
 	/** ISequencerObjectSpawner interface */
-	virtual void SpawnOrDestroyPuppetObjects( TSharedRef<FMovieSceneInstance> MovieSceneInstance, const bool bDestroyAll = false ) OVERRIDE;
-	virtual void DestroyAllPuppetObjects() OVERRIDE;
-	virtual FGuid FindSpawnableGuidForPuppetObject( UObject* Object ) OVERRIDE;
-	virtual UObject* FindPuppetObjectForSpawnableGuid( TSharedRef<FMovieSceneInstance>& MovieSceneInstance, const FGuid& Guid ) OVERRIDE;
+	virtual void SpawnOrDestroyPuppetObjects( TSharedRef<FMovieSceneInstance> MovieSceneInstance, const bool bDestroyAll = false ) override;
+	virtual void DestroyAllPuppetObjects() override;
+	virtual FGuid FindSpawnableGuidForPuppetObject( UObject* Object ) override;
+	virtual UObject* FindPuppetObjectForSpawnableGuid( TSharedRef<FMovieSceneInstance>& MovieSceneInstance, const FGuid& Guid ) override;
 
 protected:
 
diff --git a/Engine/Source/Editor/Sequencer/Private/TimeSliderController.h b/Engine/Source/Editor/Sequencer/Private/TimeSliderController.h
--- a/Engine/Source/Editor/Sequencer/Private/TimeSliderController.h
+++ b/Engine/Source/Editor/Sequencer/Private/TimeSliderController.h
@@ -14,11 +14,11 @@ public:
 	FSequencerTimeSliderController( const FTimeSliderArgs& InArgs );
 
 	/** ITimeSliderController Interface */
-	virtual int32 OnPaintTimeSlider( bool bMirrorLabels, const FGeometry& AllottedGeometry, const FSlateRect& MyClippingRect, FSlateWindowElementList& OutDrawElements, int32 LayerId, const FWidgetStyle& InWidgetStyle, bool bParentEnabled ) const OVERRIDE;
-	virtual FReply OnMouseButtonDown( TSharedRef<SWidget> WidgetOwner, const FGeometry& MyGeometry, const FPointerEvent& MouseEvent ) OVERRIDE;
-	virtual FReply OnMouseButtonUp( TSharedRef<SWidget> WidgetOwner, const FGeometry& MyGeometry, const FPointerEvent& MouseEvent ) OVERRIDE;
-	virtual FReply OnMouseMove( TSharedRef<SWidget> WidgetOwner, const FGeometry& MyGeometry, const FPointerEvent& MouseEvent ) OVERRIDE;
-	virtual FReply OnMouseWheel( TSharedRef<SWidget> WidgetOwner, const FGeometry& MyGeometry, const FPointerEvent& MouseEvent ) OVERRIDE;
+	virtual int32 OnPaintTimeSlider( bool bMirrorLabels, const FGeometry& AllottedGeometry, const FSlateRect& MyClippingRect, FSlateWindowElementList& OutDrawElements, int32 LayerId, const FWidgetStyle& InWidgetStyle, bool bParentEnabled ) const override;
+	virtual FReply OnMouseButtonDown( TSharedRef<SWidget> WidgetOwner, const FGeometry& MyGeometry, const FPointerEvent& MouseEvent ) override;
+	virtual FReply OnMouseButtonUp( TSharedRef<SWidget> WidgetOwner, const FGeometry& MyGeometry, const FPointerEvent& MouseEvent ) override;
+	virtual FReply OnMouseMove( TSharedRef<SWidget> WidgetOwner, const FGeometry& MyGeometry, const FPointerEvent& MouseEvent ) override;
+	virtual FReply OnMouseWheel( TSharedRef<SWidget> WidgetOwner, const FGeometry& MyGeometry, const FPointerEvent& MouseEvent ) override;
 
 	/**
 	 * Draws major tick lines in the section view                                                              
diff --git a/Engine/Source/Editor/SequencerWidgets/Private/STimeSlider.h b/Engine/Source/Editor/SequencerWidgets/Private/STimeSlider.h
--- a/Engine/Source/Editor/SequencerWidgets/Private/STimeSlider.h
+++ b/Engine/Source/Editor/SequencerWidgets/Private/STimeSlider.h
@@ -24,12 +24,12 @@ public:
 
 protected:
 	// SWidget interface
-	virtual FVector2D ComputeDesiredSize() const OVERRIDE;
-	virtual int32 OnPaint( const FGeometry& AllottedGeometry, const FSlateRect& MyClippingRect, FSlateWindowElementList& OutDrawElements, int32 LayerId, const FWidgetStyle& InWidgetStyle, bool bParentEnabled ) const OVERRIDE;
-	virtual FReply OnMouseButtonDown( const FGeometry& MyGeometry, const FPointerEvent& MouseEvent ) OVERRIDE;
-	virtual FReply OnMouseButtonUp( const FGeometry& MyGeometry, const FPointerEvent& MouseEvent ) OVERRIDE;
-	virtual FReply OnMouseMove( const FGeometry& MyGeometry, const FPointerEvent& MouseEvent ) OVERRIDE;
-	virtual FReply OnMouseWheel( const FGeometry& MyGeometry, const FPointerEvent& MouseEvent ) OVERRIDE;
+	virtual FVector2D ComputeDesiredSize() const override;
+	virtual int32 OnPaint( const FGeometry& AllottedGeometry, const FSlateRect& MyClippingRect, FSlateWindowElementList& OutDrawElements, int32 LayerId, const FWidgetStyle& InWidgetStyle, bool bParentEnabled ) const override;
+	virtual FReply OnMouseButtonDown( const FGeometry& MyGeometry, const FPointerEvent& MouseEvent ) override;
+	virtual FReply OnMouseButtonUp( const FGeometry& MyGeometry, const FPointerEvent& MouseEvent ) override;
+	virtual FReply OnMouseMove( const FGeometry& MyGeometry, const FPointerEvent& MouseEvent ) override;
+	virtual FReply OnMouseWheel( const FGeometry& MyGeometry, const FPointerEvent& MouseEvent ) override;
 private:
 	TSharedPtr<ITimeSliderController> TimeSliderController;
 	bool bMirrorLabels;
diff --git a/Engine/Source/Editor/SequencerWidgets/Public/SequencerWidgetsModule.h b/Engine/Source/Editor/SequencerWidgets/Public/SequencerWidgetsModule.h
--- a/Engine/Source/Editor/SequencerWidgets/Public/SequencerWidgetsModule.h
+++ b/Engine/Source/Editor/SequencerWidgets/Public/SequencerWidgetsModule.h
@@ -13,9 +13,9 @@ class FSequencerWidgetsModule : public IModuleInterface
 {
 
 public:
-	virtual void StartupModule() OVERRIDE;
+	virtual void StartupModule() override;
 
-	virtual void ShutdownModule() OVERRIDE;
+	virtual void ShutdownModule() override;
 
 	virtual TSharedRef<ITimeSlider> CreateTimeSlider( const TSharedRef<class ITimeSliderController>& InController, bool bMirrorLabels );
 };
diff --git a/Engine/Source/Editor/SoundClassEditor/Private/SoundClassEditor.h b/Engine/Source/Editor/SoundClassEditor/Private/SoundClassEditor.h
--- a/Engine/Source/Editor/SoundClassEditor/Private/SoundClassEditor.h
+++ b/Engine/Source/Editor/SoundClassEditor/Private/SoundClassEditor.h
@@ -16,8 +16,8 @@ public:
 
 	SLATE_END_ARGS()
 
-	virtual void RegisterTabSpawners(const TSharedRef<class FTabManager>& TabManager) OVERRIDE;
-	virtual void UnregisterTabSpawners(const TSharedRef<class FTabManager>& TabManager) OVERRIDE;
+	virtual void RegisterTabSpawners(const TSharedRef<class FTabManager>& TabManager) override;
+	virtual void UnregisterTabSpawners(const TSharedRef<class FTabManager>& TabManager) override;
 
 	/**
 	 * Edits the specified sound class object
@@ -31,22 +31,22 @@ public:
 	virtual ~FSoundClassEditor();
 
 	/** FGCObject interface */
-	virtual void AddReferencedObjects( FReferenceCollector& Collector ) OVERRIDE;
+	virtual void AddReferencedObjects( FReferenceCollector& Collector ) override;
 
 	/** IToolkit interface */
-	virtual FName GetToolkitFName() const OVERRIDE;
-	virtual FText GetBaseToolkitName() const OVERRIDE;
-	virtual FString GetWorldCentricTabPrefix() const OVERRIDE;
+	virtual FName GetToolkitFName() const override;
+	virtual FText GetBaseToolkitName() const override;
+	virtual FString GetWorldCentricTabPrefix() const override;
 
 	/** @return Returns the color and opacity to use for the color that appears behind the tab text for this toolkit's tab in world-centric mode. */
-	virtual FLinearColor GetWorldCentricTabColorScale() const OVERRIDE;
+	virtual FLinearColor GetWorldCentricTabColorScale() const override;
 
 	/** ISoundClassEditor interface */
-	void CreateSoundClass(class UEdGraphPin* FromPin, const FVector2D& Location, FString Name) OVERRIDE;
+	void CreateSoundClass(class UEdGraphPin* FromPin, const FVector2D& Location, FString Name) override;
 
 	/** FEditorUndoClient Interface */
-	virtual void PostUndo(bool bSuccess) OVERRIDE;
-	virtual void PostRedo(bool bSuccess) OVERRIDE { PostUndo(bSuccess); }
+	virtual void PostUndo(bool bSuccess) override;
+	virtual void PostRedo(bool bSuccess) override { PostUndo(bSuccess); }
 
 private:
 	TSharedRef<SDockTab> SpawnTab_GraphCanvas(const FSpawnTabArgs& Args);
diff --git a/Engine/Source/Editor/SoundClassEditor/Private/SoundClassEditorModule.cpp b/Engine/Source/Editor/SoundClassEditor/Private/SoundClassEditorModule.cpp
--- a/Engine/Source/Editor/SoundClassEditor/Private/SoundClassEditorModule.cpp
+++ b/Engine/Source/Editor/SoundClassEditor/Private/SoundClassEditorModule.cpp
@@ -24,7 +24,7 @@ public:
 	/**
 	 * Called right after the module DLL has been loaded and the module object has been created
 	 */
-	virtual void StartupModule() OVERRIDE
+	virtual void StartupModule() override
 	{
 		MenuExtensibilityManager = MakeShareable(new FExtensibilityManager);
 		ToolBarExtensibilityManager = MakeShareable(new FExtensibilityManager);
@@ -33,7 +33,7 @@ public:
 	/**
 	 * Called before the module is unloaded, right before the module object is destroyed.
 	 */
-	virtual void ShutdownModule() OVERRIDE
+	virtual void ShutdownModule() override
 	{
 		MenuExtensibilityManager.Reset();
 		ToolBarExtensibilityManager.Reset();
@@ -42,7 +42,7 @@ public:
 	/**
 	 * Creates a new sound class editor for a sound class object
 	 */
-	virtual TSharedRef<FAssetEditorToolkit> CreateSoundClassEditor( const EToolkitMode::Type Mode, const TSharedPtr< IToolkitHost >& InitToolkitHost, USoundClass* InSoundClass ) OVERRIDE
+	virtual TSharedRef<FAssetEditorToolkit> CreateSoundClassEditor( const EToolkitMode::Type Mode, const TSharedPtr< IToolkitHost >& InitToolkitHost, USoundClass* InSoundClass ) override
 	{
 		TSharedRef<FSoundClassEditor> NewSoundClassEditor(new FSoundClassEditor());
 		NewSoundClassEditor->InitSoundClassEditor(Mode, InitToolkitHost, InSoundClass);
diff --git a/Engine/Source/Editor/SoundCueEditor/Private/SSoundCuePalette.h b/Engine/Source/Editor/SoundCueEditor/Private/SSoundCuePalette.h
--- a/Engine/Source/Editor/SoundCueEditor/Private/SSoundCuePalette.h
+++ b/Engine/Source/Editor/SoundCueEditor/Private/SSoundCuePalette.h
@@ -17,5 +17,5 @@ public:
 
 protected:
 	/** Callback used to populate all actions list in SGraphActionMenu */
-	virtual void CollectAllActions(FGraphActionListBuilderBase& OutAllActions) OVERRIDE;
+	virtual void CollectAllActions(FGraphActionListBuilderBase& OutAllActions) override;
 };
diff --git a/Engine/Source/Editor/SoundCueEditor/Private/SoundCueEditor.h b/Engine/Source/Editor/SoundCueEditor/Private/SoundCueEditor.h
--- a/Engine/Source/Editor/SoundCueEditor/Private/SoundCueEditor.h
+++ b/Engine/Source/Editor/SoundCueEditor/Private/SoundCueEditor.h
@@ -15,8 +15,8 @@ class FSoundCueEditor : public ISoundCueEditor, public FGCObject, public FNotify
 {
 	public:
 
-	virtual void RegisterTabSpawners(const TSharedRef<class FTabManager>& TabManager) OVERRIDE;
-	virtual void UnregisterTabSpawners(const TSharedRef<class FTabManager>& TabManager) OVERRIDE;
+	virtual void RegisterTabSpawners(const TSharedRef<class FTabManager>& TabManager) override;
+	virtual void UnregisterTabSpawners(const TSharedRef<class FTabManager>& TabManager) override;
 
 	/** Destructor */
 	virtual ~FSoundCueEditor();
@@ -25,29 +25,29 @@ class FSoundCueEditor : public ISoundCueEditor, public FGCObject, public FNotify
 	void InitSoundCueEditor(const EToolkitMode::Type Mode, const TSharedPtr< class IToolkitHost >& InitToolkitHost, UObject* ObjectToEdit);
 
 	/** ISoundCueEditor interface */
-	virtual USoundCue* GetSoundCue() const OVERRIDE;
-	virtual void SetSelection(TArray<UObject*> SelectedObjects) OVERRIDE;
-	virtual bool GetBoundsForSelectedNodes(class FSlateRect& Rect, float Padding) OVERRIDE;
-	virtual int32 GetNumberOfSelectedNodes() const OVERRIDE;
-	virtual TSet<UObject*> GetSelectedNodes() const OVERRIDE;
+	virtual USoundCue* GetSoundCue() const override;
+	virtual void SetSelection(TArray<UObject*> SelectedObjects) override;
+	virtual bool GetBoundsForSelectedNodes(class FSlateRect& Rect, float Padding) override;
+	virtual int32 GetNumberOfSelectedNodes() const override;
+	virtual TSet<UObject*> GetSelectedNodes() const override;
 	
 	/** IToolkit interface */
-	virtual FName GetToolkitFName() const OVERRIDE;
-	virtual FText GetBaseToolkitName() const OVERRIDE;
-	virtual FString GetWorldCentricTabPrefix() const OVERRIDE;
-	virtual FLinearColor GetWorldCentricTabColorScale() const OVERRIDE;
+	virtual FName GetToolkitFName() const override;
+	virtual FText GetBaseToolkitName() const override;
+	virtual FString GetWorldCentricTabPrefix() const override;
+	virtual FLinearColor GetWorldCentricTabColorScale() const override;
 
-	virtual FString GetDocumentationLink() const OVERRIDE
+	virtual FString GetDocumentationLink() const override
 	{
 		return FString(TEXT("Engine/Audio/SoundCues/Editor"));
 	}
 
 	/** FGCObject interface */
-	virtual void AddReferencedObjects(FReferenceCollector& Collector) OVERRIDE;
+	virtual void AddReferencedObjects(FReferenceCollector& Collector) override;
 
 	// Begin FEditorUndoClient Interface
-	virtual void PostUndo(bool bSuccess) OVERRIDE;
-	virtual void PostRedo(bool bSuccess) OVERRIDE { PostUndo(bSuccess); }
+	virtual void PostUndo(bool bSuccess) override;
+	virtual void PostRedo(bool bSuccess) override { PostUndo(bSuccess); }
 	// End of FEditorUndoClient
 
 private:
@@ -96,9 +96,9 @@ protected:
 	/** Paste the contents of the clipboard */
 	void PasteNodes();
 	/** Paste the contents of the clipboard at a specific location */
-	virtual void PasteNodesHere(const FVector2D& Location) OVERRIDE;
+	virtual void PasteNodesHere(const FVector2D& Location) override;
 	/** Whether we are able to paste the contents of the clipboard */
-	virtual bool CanPasteNodes() const OVERRIDE;
+	virtual bool CanPasteNodes() const override;
 
 	/** Duplicate the currently selected nodes */
 	void DuplicateNodes();
@@ -113,7 +113,7 @@ protected:
 
 private:
 	/** FNotifyHook interface */
-	virtual void NotifyPostChange( const FPropertyChangedEvent& PropertyChangedEvent, UProperty* PropertyThatChanged) OVERRIDE;
+	virtual void NotifyPostChange( const FPropertyChangedEvent& PropertyChangedEvent, UProperty* PropertyThatChanged) override;
 
 	/** Creates all internal widgets for the tabs to point at */
 	void CreateInternalWidgets();
diff --git a/Engine/Source/Editor/SoundCueEditor/Private/SoundCueEditorModule.cpp b/Engine/Source/Editor/SoundCueEditor/Private/SoundCueEditorModule.cpp
--- a/Engine/Source/Editor/SoundCueEditor/Private/SoundCueEditorModule.cpp
+++ b/Engine/Source/Editor/SoundCueEditor/Private/SoundCueEditorModule.cpp
@@ -24,21 +24,21 @@ public:
 	}
 
 	/** Called right after the module DLL has been loaded and the module object has been created */
-	virtual void StartupModule() OVERRIDE
+	virtual void StartupModule() override
 	{
 		MenuExtensibilityManager = MakeShareable(new FExtensibilityManager);
 		ToolBarExtensibilityManager = MakeShareable(new FExtensibilityManager);
 	}
 
 	/** Called before the module is unloaded, right before the module object is destroyed. */
-	virtual void ShutdownModule() OVERRIDE
+	virtual void ShutdownModule() override
 	{
 		MenuExtensibilityManager.Reset();
 		ToolBarExtensibilityManager.Reset();
 	}
 
 	/** Creates a new material editor, either for a material or a material function */
-	virtual TSharedRef<ISoundCueEditor> CreateSoundCueEditor(const EToolkitMode::Type Mode, const TSharedPtr< IToolkitHost >& InitToolkitHost, USoundCue* SoundCue) OVERRIDE
+	virtual TSharedRef<ISoundCueEditor> CreateSoundCueEditor(const EToolkitMode::Type Mode, const TSharedPtr< IToolkitHost >& InitToolkitHost, USoundCue* SoundCue) override
 	{
 		TSharedRef<FSoundCueEditor> NewSoundCueEditor(new FSoundCueEditor());
 		NewSoundCueEditor->InitSoundCueEditor(Mode, InitToolkitHost, SoundCue);
diff --git a/Engine/Source/Editor/SourceControlWindows/Private/SSourceControlHistory.cpp b/Engine/Source/Editor/SourceControlWindows/Private/SSourceControlHistory.cpp
--- a/Engine/Source/Editor/SourceControlWindows/Private/SSourceControlHistory.cpp
+++ b/Engine/Source/Editor/SourceControlWindows/Private/SSourceControlHistory.cpp
@@ -376,7 +376,7 @@ public:
 	 * 
 	 * @return A new slate widget representing the dragged item
 	 */
-	virtual TSharedPtr<SWidget> GetDefaultDecorator() const OVERRIDE
+	virtual TSharedPtr<SWidget> GetDefaultDecorator() const override
 	{
 		return SNew(SBorder)
 			.BorderImage(FEditorStyle::GetBrush("Graph.ConnectorFeedback.Border"))
@@ -469,7 +469,7 @@ public:
 		);
 	}
 
-	virtual TSharedRef<SWidget> GenerateWidgetForColumn( const FName& ColumnName ) OVERRIDE
+	virtual TSharedRef<SWidget> GenerateWidgetForColumn( const FName& ColumnName ) override
 	{
 		check(RevisionListItem.IsValid());
 
diff --git a/Engine/Source/Editor/StaticMeshEditor/Private/SStaticMeshEditorViewport.h b/Engine/Source/Editor/StaticMeshEditor/Private/SStaticMeshEditorViewport.h
--- a/Engine/Source/Editor/StaticMeshEditor/Private/SStaticMeshEditorViewport.h
+++ b/Engine/Source/Editor/StaticMeshEditor/Private/SStaticMeshEditorViewport.h
@@ -20,7 +20,7 @@ public:
 	~SStaticMeshEditorViewport();
 	
 	// FGCObject interface
-	virtual void AddReferencedObjects( FReferenceCollector& Collector ) OVERRIDE;
+	virtual void AddReferencedObjects( FReferenceCollector& Collector ) override;
 	// End of FGCObject interface
 
 	/** Constructs, destroys, and updates preview mesh components based on the preview static mesh's sockets. */
@@ -69,11 +69,11 @@ public:
 
 protected:
 	/** SEditorViewport interface */
-	virtual TSharedRef<FEditorViewportClient> MakeEditorViewportClient() OVERRIDE;
-	virtual EVisibility OnGetViewportContentVisibility() const OVERRIDE;
-	virtual void BindCommands() OVERRIDE;
-	virtual void OnFocusViewportToSelection() OVERRIDE;
-	virtual TSharedPtr<SWidget> MakeViewportToolbar() OVERRIDE;
+	virtual TSharedRef<FEditorViewportClient> MakeEditorViewportClient() override;
+	virtual EVisibility OnGetViewportContentVisibility() const override;
+	virtual void BindCommands() override;
+	virtual void OnFocusViewportToSelection() override;
+	virtual TSharedPtr<SWidget> MakeViewportToolbar() override;
 private:
 	/** Determines the visibility of the viewport. */
 	bool IsVisible() const;
diff --git a/Engine/Source/Editor/StaticMeshEditor/Private/StaticMeshEditor.h b/Engine/Source/Editor/StaticMeshEditor/Private/StaticMeshEditor.h
--- a/Engine/Source/Editor/StaticMeshEditor/Private/StaticMeshEditor.h
+++ b/Engine/Source/Editor/StaticMeshEditor/Private/StaticMeshEditor.h
@@ -18,8 +18,8 @@ public:
 
 	~FStaticMeshEditor();
 
-	virtual void RegisterTabSpawners(const TSharedRef<class FTabManager>& TabManager) OVERRIDE;
-	virtual void UnregisterTabSpawners(const TSharedRef<class FTabManager>& TabManager) OVERRIDE;
+	virtual void RegisterTabSpawners(const TSharedRef<class FTabManager>& TabManager) override;
+	virtual void UnregisterTabSpawners(const TSharedRef<class FTabManager>& TabManager) override;
 
 	/**
 	 * Edits the specified static mesh object
@@ -34,31 +34,31 @@ public:
 	TSharedRef<class IDetailCustomization> MakeStaticMeshDetails();
 
 	// FGCObject interface
-	virtual void AddReferencedObjects( FReferenceCollector& Collector ) OVERRIDE;
+	virtual void AddReferencedObjects( FReferenceCollector& Collector ) override;
 	// End of FGCObject interface
 
 	/** IToolkit interface */
-	virtual FName GetToolkitFName() const OVERRIDE;
-	virtual FText GetBaseToolkitName() const OVERRIDE;
-	virtual FString GetWorldCentricTabPrefix() const OVERRIDE;
+	virtual FName GetToolkitFName() const override;
+	virtual FText GetBaseToolkitName() const override;
+	virtual FString GetWorldCentricTabPrefix() const override;
 
 	/** @return the documentation location for this editor */
-	virtual FString GetDocumentationLink() const OVERRIDE
+	virtual FString GetDocumentationLink() const override
 	{
 		return FString(TEXT("Engine/Content/Types/StaticMeshes/Editor"));
 	}
 
 	/** @return Returns the color and opacity to use for the color that appears behind the tab text for this toolkit's tab in world-centric mode. */
-	virtual FLinearColor GetWorldCentricTabColorScale() const OVERRIDE;
+	virtual FLinearColor GetWorldCentricTabColorScale() const override;
 
 	/** IStaticMeshEditor interface */
 	virtual UStaticMesh* GetStaticMesh() { return StaticMesh; }
-	virtual UStaticMeshComponent* GetStaticMeshComponent() const OVERRIDE;
+	virtual UStaticMeshComponent* GetStaticMeshComponent() const override;
 
-	virtual UStaticMeshSocket* GetSelectedSocket() const OVERRIDE ;
-	virtual void SetSelectedSocket(UStaticMeshSocket* InSelectedSocket) OVERRIDE;
-	virtual void DuplicateSelectedSocket() OVERRIDE;
-	virtual void RequestRenameSelectedSocket() OVERRIDE;
+	virtual UStaticMeshSocket* GetSelectedSocket() const override ;
+	virtual void SetSelectedSocket(UStaticMeshSocket* InSelectedSocket) override;
+	virtual void DuplicateSelectedSocket() override;
+	virtual void RequestRenameSelectedSocket() override;
 
 	virtual bool IsPrimValid(const FPrimData& InPrimData) const override;
 	virtual bool HasSelectedPrims() const;
@@ -76,34 +76,34 @@ public:
 	FTransform GetPrimTransform(const FPrimData& InPrimData) const;
 	void SetPrimTransform(const FPrimData& InPrimData, const FTransform& InPrimTransform) const;
 
-	virtual int32 GetNumTriangles( int32 LODLevel = 0 ) const OVERRIDE;
-	virtual int32 GetNumVertices( int32 LODLevel = 0 ) const OVERRIDE;
-	virtual int32 GetNumUVChannels( int32 LODLevel = 0 ) const OVERRIDE;
+	virtual int32 GetNumTriangles( int32 LODLevel = 0 ) const override;
+	virtual int32 GetNumVertices( int32 LODLevel = 0 ) const override;
+	virtual int32 GetNumUVChannels( int32 LODLevel = 0 ) const override;
 
-	virtual int32 GetCurrentUVChannel() OVERRIDE;
-	virtual int32 GetCurrentLODLevel() OVERRIDE;
-	virtual int32 GetCurrentLODIndex() OVERRIDE;
+	virtual int32 GetCurrentUVChannel() override;
+	virtual int32 GetCurrentLODLevel() override;
+	virtual int32 GetCurrentLODIndex() override;
 
-	virtual void RefreshTool() OVERRIDE;
-	virtual void RefreshViewport() OVERRIDE;
+	virtual void RefreshTool() override;
+	virtual void RefreshViewport() override;
 
 	/** This is called when Apply is pressed in the dialog. Does the actual processing. */
-	virtual void DoDecomp(int32 InMaxHullCount, int32 InMaxHullVerts) OVERRIDE;
+	virtual void DoDecomp(int32 InMaxHullCount, int32 InMaxHullVerts) override;
 
-	virtual TSet< int32 >& GetSelectedEdges() OVERRIDE;
+	virtual TSet< int32 >& GetSelectedEdges() override;
 	// End of IStaticMeshEditor
 	
 	/** Extends the toolbar menu to include static mesh editor options */
 	void ExtendMenu();
 
 	/** Registers a delegate to be called after an Undo operation */
-	virtual void RegisterOnPostUndo( const FOnPostUndo& Delegate ) OVERRIDE;
+	virtual void RegisterOnPostUndo( const FOnPostUndo& Delegate ) override;
 
 	/** Unregisters a delegate to be called after an Undo operation */
-	virtual void UnregisterOnPostUndo( SWidget* Widget ) OVERRIDE;
+	virtual void UnregisterOnPostUndo( SWidget* Widget ) override;
 	
 	/** From FNotifyHook */
-	virtual void NotifyPostChange( const FPropertyChangedEvent& PropertyChangedEvent, UProperty* PropertyThatChanged ) OVERRIDE;
+	virtual void NotifyPostChange( const FPropertyChangedEvent& PropertyChangedEvent, UProperty* PropertyThatChanged ) override;
 	
 	/** Get the names of the LOD for menus */
 	TArray< TSharedPtr< FString > >& GetLODLevels() { return LODLevels; }
@@ -223,12 +223,12 @@ private:
 	void OnConvexDecomposition();
 
 	// Begin FAssetEditorToolkit interface.
-	virtual bool OnRequestClose() OVERRIDE;
+	virtual bool OnRequestClose() override;
 	// End FAssetEditorToolkit interface.
 
 	// Begin FEditorUndoClient Interface
-	virtual void PostUndo( bool bSuccess ) OVERRIDE;
-	virtual void PostRedo( bool bSuccess ) OVERRIDE;
+	virtual void PostUndo( bool bSuccess ) override;
+	virtual void PostRedo( bool bSuccess ) override;
 	// End of FEditorUndoClient
 
 	/** Undo Action**/
diff --git a/Engine/Source/Editor/StaticMeshEditor/Private/StaticMeshEditorModule.cpp b/Engine/Source/Editor/StaticMeshEditor/Private/StaticMeshEditorModule.cpp
--- a/Engine/Source/Editor/StaticMeshEditor/Private/StaticMeshEditorModule.cpp
+++ b/Engine/Source/Editor/StaticMeshEditor/Private/StaticMeshEditorModule.cpp
@@ -21,7 +21,7 @@ public:
 	/**
 	 * Called right after the module DLL has been loaded and the module object has been created
 	 */
-	virtual void StartupModule() OVERRIDE
+	virtual void StartupModule() override
 	{
 		MenuExtensibilityManager = MakeShareable(new FExtensibilityManager);
 		ToolBarExtensibilityManager = MakeShareable(new FExtensibilityManager);
@@ -30,7 +30,7 @@ public:
 	/**
 	 * Called before the module is unloaded, right before the module object is destroyed.
 	 */
-	virtual void ShutdownModule() OVERRIDE
+	virtual void ShutdownModule() override
 	{
 		MenuExtensibilityManager.Reset();
 		ToolBarExtensibilityManager.Reset();
@@ -39,7 +39,7 @@ public:
 	/**
 	 * Creates a new StaticMesh editor for a StaticMesh
 	 */
-	virtual TSharedRef<IStaticMeshEditor> CreateStaticMeshEditor( const EToolkitMode::Type Mode, const TSharedPtr< IToolkitHost >& InitToolkitHost, UStaticMesh* StaticMesh ) OVERRIDE
+	virtual TSharedRef<IStaticMeshEditor> CreateStaticMeshEditor( const EToolkitMode::Type Mode, const TSharedPtr< IToolkitHost >& InitToolkitHost, UStaticMesh* StaticMesh ) override
 	{
 		TSharedRef<FStaticMeshEditor> NewStaticMeshEditor(new FStaticMeshEditor());
 		NewStaticMeshEditor->InitStaticMeshEditor(Mode, InitToolkitHost,StaticMesh);
diff --git a/Engine/Source/Editor/StaticMeshEditor/Private/StaticMeshEditorTools.h b/Engine/Source/Editor/StaticMeshEditor/Private/StaticMeshEditorTools.h
--- a/Engine/Source/Editor/StaticMeshEditor/Private/StaticMeshEditorTools.h
+++ b/Engine/Source/Editor/StaticMeshEditor/Private/StaticMeshEditorTools.h
@@ -26,7 +26,7 @@ public:
 	~FStaticMeshDetails();
 
 	/** IDetailCustomization interface */
-	virtual void CustomizeDetails( class IDetailLayoutBuilder& DetailBuilder ) OVERRIDE;
+	virtual void CustomizeDetails( class IDetailLayoutBuilder& DetailBuilder ) override;
 
 	/** @return true if settings have been changed and need to be applied to the static mesh */
 	bool IsApplyNeeded() const;
@@ -124,13 +124,13 @@ public:
 
 private:
 	/** IDetailCustomNodeBuilder Interface*/
-	virtual void SetOnRebuildChildren( FSimpleDelegate InOnRegenerateChildren ) OVERRIDE {}
-	virtual void GenerateHeaderRowContent( FDetailWidgetRow& NodeRow ) OVERRIDE;
-	virtual void GenerateChildContent( IDetailChildrenBuilder& ChildrenBuilder ) OVERRIDE;
-	virtual void Tick( float DeltaTime ) OVERRIDE{}
-	virtual bool RequiresTick() const OVERRIDE { return false; }
-	virtual FName GetName() const OVERRIDE { static FName MeshBuildSettings("MeshBuildSettings"); return MeshBuildSettings; }
-	virtual bool InitiallyCollapsed() const OVERRIDE { return true; }
+	virtual void SetOnRebuildChildren( FSimpleDelegate InOnRegenerateChildren ) override {}
+	virtual void GenerateHeaderRowContent( FDetailWidgetRow& NodeRow ) override;
+	virtual void GenerateChildContent( IDetailChildrenBuilder& ChildrenBuilder ) override;
+	virtual void Tick( float DeltaTime ) override{}
+	virtual bool RequiresTick() const override { return false; }
+	virtual FName GetName() const override { static FName MeshBuildSettings("MeshBuildSettings"); return MeshBuildSettings; }
+	virtual bool InitiallyCollapsed() const override { return true; }
 
 	FReply OnApplyChanges();
 	ESlateCheckBoxState::Type ShouldRecomputeNormals() const;
@@ -166,13 +166,13 @@ public:
 	void UpdateSettings(const FMeshReductionSettings& InSettings);
 private:
 	/** IDetailCustomNodeBuilder Interface*/
-	virtual void SetOnRebuildChildren( FSimpleDelegate InOnRegenerateChildren ) OVERRIDE {}
-	virtual void GenerateHeaderRowContent( FDetailWidgetRow& NodeRow ) OVERRIDE;
-	virtual void GenerateChildContent( IDetailChildrenBuilder& ChildrenBuilder ) OVERRIDE;
-	virtual void Tick( float DeltaTime ) OVERRIDE{}
-	virtual bool RequiresTick() const OVERRIDE { return false; }
-	virtual FName GetName() const OVERRIDE { static FName MeshReductionSettings("MeshReductionSettings"); return MeshReductionSettings; }
-	virtual bool InitiallyCollapsed() const OVERRIDE { return true; }
+	virtual void SetOnRebuildChildren( FSimpleDelegate InOnRegenerateChildren ) override {}
+	virtual void GenerateHeaderRowContent( FDetailWidgetRow& NodeRow ) override;
+	virtual void GenerateChildContent( IDetailChildrenBuilder& ChildrenBuilder ) override;
+	virtual void Tick( float DeltaTime ) override{}
+	virtual bool RequiresTick() const override { return false; }
+	virtual FName GetName() const override { static FName MeshReductionSettings("MeshReductionSettings"); return MeshReductionSettings; }
+	virtual bool InitiallyCollapsed() const override { return true; }
 
 	FReply OnApplyChanges();
 	float GetPercentTriangles() const;
diff --git a/Engine/Source/Editor/StaticMeshEditor/Private/StaticMeshEditorViewportClient.h b/Engine/Source/Editor/StaticMeshEditor/Private/StaticMeshEditorViewportClient.h
--- a/Engine/Source/Editor/StaticMeshEditor/Private/StaticMeshEditorViewportClient.h
+++ b/Engine/Source/Editor/StaticMeshEditor/Private/StaticMeshEditorViewportClient.h
@@ -11,22 +11,22 @@ public:
 	FStaticMeshEditorViewportClient(TWeakPtr<IStaticMeshEditor> InStaticMeshEditor, FPreviewScene& InPreviewScene, UStaticMesh* InPreviewStaticMesh, UStaticMeshComponent* InPreviewStaticMeshComponent);
 
 	// FEditorViewportClient interface
-	virtual void MouseMove(FViewport* Viewport,int32 x, int32 y) OVERRIDE;
-	virtual bool InputKey(FViewport* Viewport, int32 ControllerId, FKey Key, EInputEvent Event, float AmountDepressed = 1.f, bool bGamepad=false) OVERRIDE;
-	virtual bool InputAxis(FViewport* Viewport, int32 ControllerId, FKey Key, float Delta, float DeltaTime, int32 NumSamples=1, bool bGamepad=false) OVERRIDE;
-	virtual void ProcessClick(class FSceneView& View, class HHitProxy* HitProxy, FKey Key, EInputEvent Event, uint32 HitX, uint32 HitY) OVERRIDE;
-	virtual void Tick(float DeltaSeconds) OVERRIDE;
-	virtual void Draw(const FSceneView* View,FPrimitiveDrawInterface* PDI) OVERRIDE;
-	virtual void DrawCanvas( FViewport& InViewport, FSceneView& View, FCanvas& Canvas ) OVERRIDE;
-	virtual bool InputWidgetDelta( FViewport* Viewport, EAxisList::Type CurrentAxis, FVector& Drag, FRotator& Rot, FVector& Scale ) OVERRIDE;
-	virtual void TrackingStarted( const struct FInputEventState& InInputState, bool bIsDragging, bool bNudge ) OVERRIDE;
-	virtual void TrackingStopped() OVERRIDE;
-	virtual FWidget::EWidgetMode GetWidgetMode() const OVERRIDE;
+	virtual void MouseMove(FViewport* Viewport,int32 x, int32 y) override;
+	virtual bool InputKey(FViewport* Viewport, int32 ControllerId, FKey Key, EInputEvent Event, float AmountDepressed = 1.f, bool bGamepad=false) override;
+	virtual bool InputAxis(FViewport* Viewport, int32 ControllerId, FKey Key, float Delta, float DeltaTime, int32 NumSamples=1, bool bGamepad=false) override;
+	virtual void ProcessClick(class FSceneView& View, class HHitProxy* HitProxy, FKey Key, EInputEvent Event, uint32 HitX, uint32 HitY) override;
+	virtual void Tick(float DeltaSeconds) override;
+	virtual void Draw(const FSceneView* View,FPrimitiveDrawInterface* PDI) override;
+	virtual void DrawCanvas( FViewport& InViewport, FSceneView& View, FCanvas& Canvas ) override;
+	virtual bool InputWidgetDelta( FViewport* Viewport, EAxisList::Type CurrentAxis, FVector& Drag, FRotator& Rot, FVector& Scale ) override;
+	virtual void TrackingStarted( const struct FInputEventState& InInputState, bool bIsDragging, bool bNudge ) override;
+	virtual void TrackingStopped() override;
+	virtual FWidget::EWidgetMode GetWidgetMode() const override;
 	virtual FVector GetWidgetLocation() const;
 	virtual FMatrix GetWidgetCoordSystem() const;
 	virtual ECoordSystem GetWidgetCoordSystemSpace() const { return COORD_Local; }
-	virtual bool ShouldOrbitCamera() const OVERRIDE;
-	virtual FSceneView* CalcSceneView(FSceneViewFamily* ViewFamily) OVERRIDE;
+	virtual bool ShouldOrbitCamera() const override;
+	virtual FSceneView* CalcSceneView(FSceneViewFamily* ViewFamily) override;
 
 	/** 
 	 *	Updates the static mesh and static mesh component being used in the Static Mesh Editor.
@@ -108,7 +108,7 @@ public:
 
 protected:
 	// FEditorViewportClient interface
-	virtual void PerspectiveCameraMoved() OVERRIDE;
+	virtual void PerspectiveCameraMoved() override;
 
 private:
 	/** The Simplygon logo to be drawn when Simplygon has been used on the static mesh. */
diff --git a/Engine/Source/Editor/StatsViewer/Private/ActorArrayHyperlinkColumn.cpp b/Engine/Source/Editor/StatsViewer/Private/ActorArrayHyperlinkColumn.cpp
--- a/Engine/Source/Editor/StatsViewer/Private/ActorArrayHyperlinkColumn.cpp
+++ b/Engine/Source/Editor/StatsViewer/Private/ActorArrayHyperlinkColumn.cpp
@@ -48,7 +48,7 @@ public:
 		}
 	}
 
-	virtual TSharedRef< class SWidget > ConstructDisplayWidget() OVERRIDE
+	virtual TSharedRef< class SWidget > ConstructDisplayWidget() override
 	{
 		return SNew( SHorizontalBox )
 			
diff --git a/Engine/Source/Editor/StatsViewer/Private/ActorArrayHyperlinkColumn.h b/Engine/Source/Editor/StatsViewer/Private/ActorArrayHyperlinkColumn.h
--- a/Engine/Source/Editor/StatsViewer/Private/ActorArrayHyperlinkColumn.h
+++ b/Engine/Source/Editor/StatsViewer/Private/ActorArrayHyperlinkColumn.h
@@ -12,9 +12,9 @@ class FActorArrayHyperlinkColumn : public IPropertyTableCustomColumn
 {
 public:
 	/** Begin IPropertyTableCustomColumn interface */
-	virtual bool Supports( const TSharedRef< IPropertyTableColumn >& Column, const TSharedRef< IPropertyTableUtilities >& Utilities ) const OVERRIDE;
-	virtual TSharedPtr< SWidget > CreateColumnLabel( const TSharedRef< IPropertyTableColumn >& Column, const TSharedRef< IPropertyTableUtilities >& Utilities, const FName& Style ) const OVERRIDE;
-	virtual TSharedPtr< IPropertyTableCellPresenter > CreateCellPresenter( const TSharedRef< IPropertyTableCell >& Cell, const TSharedRef< IPropertyTableUtilities >& Utilities, const FName& Style ) const OVERRIDE;
+	virtual bool Supports( const TSharedRef< IPropertyTableColumn >& Column, const TSharedRef< IPropertyTableUtilities >& Utilities ) const override;
+	virtual TSharedPtr< SWidget > CreateColumnLabel( const TSharedRef< IPropertyTableColumn >& Column, const TSharedRef< IPropertyTableUtilities >& Utilities, const FName& Style ) const override;
+	virtual TSharedPtr< IPropertyTableCellPresenter > CreateCellPresenter( const TSharedRef< IPropertyTableCell >& Cell, const TSharedRef< IPropertyTableUtilities >& Utilities, const FName& Style ) const override;
 	/** End IPropertyTableCustomColumn interface */
 };
 
diff --git a/Engine/Source/Editor/StatsViewer/Private/ObjectHyperlinkColumn.cpp b/Engine/Source/Editor/StatsViewer/Private/ObjectHyperlinkColumn.cpp
--- a/Engine/Source/Editor/StatsViewer/Private/ObjectHyperlinkColumn.cpp
+++ b/Engine/Source/Editor/StatsViewer/Private/ObjectHyperlinkColumn.cpp
@@ -30,7 +30,7 @@ public:
 		}
 	}
 
-	virtual TSharedRef< class SWidget > ConstructDisplayWidget() OVERRIDE
+	virtual TSharedRef< class SWidget > ConstructDisplayWidget() override
 	{
 		return OnGenerateWidget.Execute(GetValueAsText(), Object);
 	}
diff --git a/Engine/Source/Editor/StatsViewer/Private/SStatsViewer.h b/Engine/Source/Editor/StatsViewer/Private/SStatsViewer.h
--- a/Engine/Source/Editor/StatsViewer/Private/SStatsViewer.h
+++ b/Engine/Source/Editor/StatsViewer/Private/SStatsViewer.h
@@ -28,10 +28,10 @@ public:
 	virtual ~SStatsViewer();
 
 	/** SWidget interface */
-	virtual void Tick( const FGeometry& AllottedGeometry, const double InCurrentTime, const float InDeltaTime ) OVERRIDE;
+	virtual void Tick( const FGeometry& AllottedGeometry, const double InCurrentTime, const float InDeltaTime ) override;
 
 	/** Begin IStatsViewer interface */
-	virtual void Refresh() OVERRIDE;
+	virtual void Refresh() override;
 	TSharedPtr< class IPropertyTable > GetPropertyTable();
 	int32 GetObjectSetIndex() const;
 	/** End IStatsViewer interface */
diff --git a/Engine/Source/Editor/StatsViewer/Private/StatsPages/LightingBuildInfoStatsPage.h b/Engine/Source/Editor/StatsViewer/Private/StatsPages/LightingBuildInfoStatsPage.h
--- a/Engine/Source/Editor/StatsViewer/Private/StatsPages/LightingBuildInfoStatsPage.h
+++ b/Engine/Source/Editor/StatsViewer/Private/StatsPages/LightingBuildInfoStatsPage.h
@@ -12,10 +12,10 @@ public:
 	static FLightingBuildInfoStatsPage& Get();
 
 	/** Begin IStatsPage interface */
-	virtual void Clear() OVERRIDE;
-	virtual void AddEntry( UObject* InEntry ) OVERRIDE;
-	virtual void Generate( TArray< TWeakObjectPtr<UObject> >& OutObjects ) const OVERRIDE;
-	virtual void GenerateTotals( const TArray< TWeakObjectPtr<UObject> >& InObjects, TMap<FString, FText>& OutTotals ) const OVERRIDE;
+	virtual void Clear() override;
+	virtual void AddEntry( UObject* InEntry ) override;
+	virtual void Generate( TArray< TWeakObjectPtr<UObject> >& OutObjects ) const override;
+	virtual void GenerateTotals( const TArray< TWeakObjectPtr<UObject> >& InObjects, TMap<FString, FText>& OutTotals ) const override;
 	/** End IStatsPage interface */
 
 private:
diff --git a/Engine/Source/Editor/StatsViewer/Private/StatsPages/PrimitiveStatsPage.h b/Engine/Source/Editor/StatsViewer/Private/StatsPages/PrimitiveStatsPage.h
--- a/Engine/Source/Editor/StatsViewer/Private/StatsPages/PrimitiveStatsPage.h
+++ b/Engine/Source/Editor/StatsViewer/Private/StatsPages/PrimitiveStatsPage.h
@@ -12,10 +12,10 @@ public:
 	static FPrimitiveStatsPage& Get();
 
 	/** Begin IStatsPage interface */
-	virtual void Generate( TArray< TWeakObjectPtr<UObject> >& OutObjects ) const OVERRIDE;
-	virtual void GenerateTotals( const TArray< TWeakObjectPtr<UObject> >& InObjects, TMap<FString, FText>& OutTotals ) const OVERRIDE;
-	virtual void OnShow( TWeakPtr< class IStatsViewer > InParentStatsViewer ) OVERRIDE;
-	virtual void OnHide() OVERRIDE;
+	virtual void Generate( TArray< TWeakObjectPtr<UObject> >& OutObjects ) const override;
+	virtual void GenerateTotals( const TArray< TWeakObjectPtr<UObject> >& InObjects, TMap<FString, FText>& OutTotals ) const override;
+	virtual void OnShow( TWeakPtr< class IStatsViewer > InParentStatsViewer ) override;
+	virtual void OnHide() override;
 	/** End IStatsPage interface */
 
 private:
diff --git a/Engine/Source/Editor/StatsViewer/Private/StatsPages/StaticMeshLightingInfoStatsPage.h b/Engine/Source/Editor/StatsViewer/Private/StatsPages/StaticMeshLightingInfoStatsPage.h
--- a/Engine/Source/Editor/StatsViewer/Private/StatsPages/StaticMeshLightingInfoStatsPage.h
+++ b/Engine/Source/Editor/StatsViewer/Private/StatsPages/StaticMeshLightingInfoStatsPage.h
@@ -31,11 +31,11 @@ public:
 	static FStaticMeshLightingInfoStatsPage& Get();
 
 	/** Begin IStatsPage interface */
-	virtual void Generate( TArray< TWeakObjectPtr<UObject> >& OutObjects ) const OVERRIDE;
-	virtual void GenerateTotals( const TArray< TWeakObjectPtr<UObject> >& InObjects, TMap<FString, FText>& OutTotals ) const OVERRIDE;
-	virtual TSharedPtr<SWidget> GetCustomWidget(TWeakPtr< class IStatsViewer > InParentStatsViewer) OVERRIDE;
-	virtual void OnShow( TWeakPtr< class IStatsViewer > InParentStatsViewer ) OVERRIDE;
-	virtual void OnHide() OVERRIDE;
+	virtual void Generate( TArray< TWeakObjectPtr<UObject> >& OutObjects ) const override;
+	virtual void GenerateTotals( const TArray< TWeakObjectPtr<UObject> >& InObjects, TMap<FString, FText>& OutTotals ) const override;
+	virtual TSharedPtr<SWidget> GetCustomWidget(TWeakPtr< class IStatsViewer > InParentStatsViewer) override;
+	virtual void OnShow( TWeakPtr< class IStatsViewer > InParentStatsViewer ) override;
+	virtual void OnHide() override;
 	/** End IStatsPage interface */
 
 private:
diff --git a/Engine/Source/Editor/StatsViewer/Private/StatsPages/TextureStatsPage.h b/Engine/Source/Editor/StatsViewer/Private/StatsPages/TextureStatsPage.h
--- a/Engine/Source/Editor/StatsViewer/Private/StatsPages/TextureStatsPage.h
+++ b/Engine/Source/Editor/StatsViewer/Private/StatsPages/TextureStatsPage.h
@@ -12,10 +12,10 @@ public:
 	static FTextureStatsPage& Get();
 
 	/** Begin IStatsPage interface */
-	virtual void Generate( TArray< TWeakObjectPtr<UObject> >& OutObjects ) const OVERRIDE;
-	virtual void GenerateTotals( const TArray< TWeakObjectPtr<UObject> >& InObjects, TMap<FString, FText>& OutTotals ) const OVERRIDE;
-	virtual void OnShow( TWeakPtr< class IStatsViewer > InParentStatsViewer ) OVERRIDE;
-	virtual void OnHide() OVERRIDE;
+	virtual void Generate( TArray< TWeakObjectPtr<UObject> >& OutObjects ) const override;
+	virtual void GenerateTotals( const TArray< TWeakObjectPtr<UObject> >& InObjects, TMap<FString, FText>& OutTotals ) const override;
+	virtual void OnShow( TWeakPtr< class IStatsViewer > InParentStatsViewer ) override;
+	virtual void OnHide() override;
 	/** End IStatsPage interface */
 
 private:
diff --git a/Engine/Source/Editor/StatsViewer/Public/StatsCellPresenter.h b/Engine/Source/Editor/StatsViewer/Public/StatsCellPresenter.h
--- a/Engine/Source/Editor/StatsViewer/Public/StatsCellPresenter.h
+++ b/Engine/Source/Editor/StatsViewer/Public/StatsCellPresenter.h
@@ -11,37 +11,37 @@ class FStatsCellPresenter : public IPropertyTableCellPresenter
 {
 public:
 	/** Begin IPropertyTableCellPresenter interface */
-	virtual bool RequiresDropDown() OVERRIDE
+	virtual bool RequiresDropDown() override
 	{
 		return false;
 	}
 
-	virtual TSharedRef< class SWidget > ConstructEditModeCellWidget() OVERRIDE
+	virtual TSharedRef< class SWidget > ConstructEditModeCellWidget() override
 	{
 		return ConstructDisplayWidget();
 	}
 
-	virtual TSharedRef< class SWidget > ConstructEditModeDropDownWidget() OVERRIDE
+	virtual TSharedRef< class SWidget > ConstructEditModeDropDownWidget() override
 	{
 		return SNullWidget::NullWidget;
 	}
 
-	virtual TSharedRef< class SWidget > WidgetToFocusOnEdit() OVERRIDE
+	virtual TSharedRef< class SWidget > WidgetToFocusOnEdit() override
 	{
 		return SNullWidget::NullWidget;
 	}
 
-	virtual bool HasReadOnlyEditMode() OVERRIDE 
+	virtual bool HasReadOnlyEditMode() override 
 	{ 
 		return true; 
 	}
 
-	virtual FString GetValueAsString() OVERRIDE
+	virtual FString GetValueAsString() override
 	{
 		return Text.ToString();
 	}
 
-	virtual FText GetValueAsText() OVERRIDE
+	virtual FText GetValueAsText() override
 	{
 		return Text;
 	}
diff --git a/Engine/Source/Editor/StatsViewer/Public/StatsPage.h b/Engine/Source/Editor/StatsViewer/Public/StatsPage.h
--- a/Engine/Source/Editor/StatsViewer/Public/StatsPage.h
+++ b/Engine/Source/Editor/StatsViewer/Public/StatsPage.h
@@ -28,42 +28,42 @@ public:
 	}
 
 	/** Begin IStatsPage interface */
-	virtual void Show( bool bInShow = true ) OVERRIDE
+	virtual void Show( bool bInShow = true ) override
 	{
 		bShow = bInShow;
 	}
 
-	virtual bool IsShowPending() const OVERRIDE
+	virtual bool IsShowPending() const override
 	{
 		return bShow;
 	}
 
-	virtual void Refresh( bool bInRefresh = true ) OVERRIDE
+	virtual void Refresh( bool bInRefresh = true ) override
 	{
 		bRefresh = bInRefresh;
 	}
 
-	virtual bool IsRefreshPending() const OVERRIDE
+	virtual bool IsRefreshPending() const override
 	{
 		return bRefresh;
 	}
 
-	virtual FName GetName() const OVERRIDE
+	virtual FName GetName() const override
 	{
 		return Entry::StaticClass()->GetFName();
 	}
 
-	virtual const FText GetDisplayName() const OVERRIDE
+	virtual const FText GetDisplayName() const override
 	{
 		return Entry::StaticClass()->GetDisplayNameText();
 	}
 
-	virtual const FText GetToolTip() const OVERRIDE
+	virtual const FText GetToolTip() const override
 	{
 		return Entry::StaticClass()->GetToolTipText();
 	}
 
-	virtual int32 GetObjectSetCount() const OVERRIDE
+	virtual int32 GetObjectSetCount() const override
 	{
 		if(ObjectSetEnum != NULL)
 		{
@@ -72,7 +72,7 @@ public:
 		return 1;
 	}
 
-	virtual FString GetObjectSetName( int32 InObjectSetIndex ) const OVERRIDE
+	virtual FString GetObjectSetName( int32 InObjectSetIndex ) const override
 	{
 		if(ObjectSetEnum != NULL)
 		{
@@ -84,7 +84,7 @@ public:
 		return EmptyString;
 	}
 
-	virtual FString GetObjectSetToolTip( int32 InObjectSetIndex ) const OVERRIDE
+	virtual FString GetObjectSetToolTip( int32 InObjectSetIndex ) const override
 	{
 		if(ObjectSetEnum != NULL)
 		{
@@ -96,27 +96,27 @@ public:
 		return EmptyString;
 	}
 
-	virtual UClass* GetEntryClass() const OVERRIDE
+	virtual UClass* GetEntryClass() const override
 	{
 		return Entry::StaticClass();
 	}
 
-	virtual TSharedPtr<SWidget> GetCustomWidget( TWeakPtr< class IStatsViewer > InParentStatsViewer ) OVERRIDE
+	virtual TSharedPtr<SWidget> GetCustomWidget( TWeakPtr< class IStatsViewer > InParentStatsViewer ) override
 	{
 		return NULL;
 	}
 
-	virtual void SetSelectedObjectSet( int32 InObjectSetIndex ) OVERRIDE
+	virtual void SetSelectedObjectSet( int32 InObjectSetIndex ) override
 	{
 		ObjectSetIndex = InObjectSetIndex;
 	}
 
-	virtual int32 GetSelectedObjectSet() const OVERRIDE
+	virtual int32 GetSelectedObjectSet() const override
 	{
 		return ObjectSetIndex;
 	}
 
-	virtual void GetCustomColumns(TArray< TSharedRef< class IPropertyTableCustomColumn > >& OutCustomColumns) const OVERRIDE
+	virtual void GetCustomColumns(TArray< TSharedRef< class IPropertyTableCustomColumn > >& OutCustomColumns) const override
 	{
 	}
 	/** End IStatsPage interface */
diff --git a/Engine/Source/Editor/SwarmInterface/Private/SwarmInterfaceModule.cpp b/Engine/Source/Editor/SwarmInterface/Private/SwarmInterfaceModule.cpp
--- a/Engine/Source/Editor/SwarmInterface/Private/SwarmInterfaceModule.cpp
+++ b/Engine/Source/Editor/SwarmInterface/Private/SwarmInterfaceModule.cpp
@@ -11,9 +11,9 @@ class FSwarmInterfaceModule : public IModuleInterface
 
 public:
 
-	virtual void StartupModule() OVERRIDE {}
+	virtual void StartupModule() override {}
 
-	virtual void ShutdownModule() OVERRIDE {}
+	virtual void ShutdownModule() override {}
 };
 
 // Dummy class initialization
diff --git a/Engine/Source/Editor/TaskBrowser/Public/TaskDataManager.h b/Engine/Source/Editor/TaskBrowser/Public/TaskDataManager.h
--- a/Engine/Source/Editor/TaskBrowser/Public/TaskDataManager.h
+++ b/Engine/Source/Editor/TaskBrowser/Public/TaskDataManager.h
@@ -165,7 +165,7 @@ public:
 	 */
 	virtual void Tick( float DeltaTime );
 
-	virtual TStatId GetStatId() const OVERRIDE;
+	virtual TStatId GetStatId() const override;
 
 	/**
 	 * Sets the connection settings used to connect and login to the task database server
diff --git a/Engine/Source/Editor/TextureAlignMode/Public/TextureAlignEdMode.h b/Engine/Source/Editor/TextureAlignMode/Public/TextureAlignEdMode.h
--- a/Engine/Source/Editor/TextureAlignMode/Public/TextureAlignEdMode.h
+++ b/Engine/Source/Editor/TextureAlignMode/Public/TextureAlignEdMode.h
@@ -14,12 +14,12 @@ public:
 	/**
 	 * Called right after the module's DLL has been loaded and the module object has been created
 	 */
-	virtual void StartupModule() OVERRIDE;
+	virtual void StartupModule() override;
 
 	/**
 	 * Called before the module is unloaded, right before the module object is destroyed.
 	 */
-	virtual void ShutdownModule() OVERRIDE;};
+	virtual void ShutdownModule() override;};
 
 /**
  * Allows texture alignment on BSP surfaces via the widget.
diff --git a/Engine/Source/Editor/TextureAlignMode/Public/TextureAlignMode.h b/Engine/Source/Editor/TextureAlignMode/Public/TextureAlignMode.h
--- a/Engine/Source/Editor/TextureAlignMode/Public/TextureAlignMode.h
+++ b/Engine/Source/Editor/TextureAlignMode/Public/TextureAlignMode.h
@@ -8,13 +8,13 @@
 class FTextureAlignMode : public FModeToolkit
 {
 public:
-	virtual void RegisterTabSpawners(const TSharedRef<class FTabManager>& TabManager) OVERRIDE;
-	virtual void UnregisterTabSpawners(const TSharedRef<class FTabManager>& TabManager) OVERRIDE;
+	virtual void RegisterTabSpawners(const TSharedRef<class FTabManager>& TabManager) override;
+	virtual void UnregisterTabSpawners(const TSharedRef<class FTabManager>& TabManager) override;
 
 	/** IToolkit interface */
-	virtual FName GetToolkitFName() const OVERRIDE;
-	virtual FText GetBaseToolkitName() const OVERRIDE;
-	virtual class FEdMode* GetEditorMode() const OVERRIDE;
+	virtual FName GetToolkitFName() const override;
+	virtual FText GetBaseToolkitName() const override;
+	virtual class FEdMode* GetEditorMode() const override;
 
 private:
 
diff --git a/Engine/Source/Editor/TextureEditor/Private/Models/TextureEditorViewportClient.h b/Engine/Source/Editor/TextureEditor/Private/Models/TextureEditorViewportClient.h
--- a/Engine/Source/Editor/TextureEditor/Private/Models/TextureEditorViewportClient.h
+++ b/Engine/Source/Editor/TextureEditor/Private/Models/TextureEditorViewportClient.h
@@ -13,12 +13,12 @@ public:
 	~FTextureEditorViewportClient();
 
 	/** FViewportClient interface */
-	virtual void Draw(FViewport* Viewport, FCanvas* Canvas) OVERRIDE;
-	virtual bool InputKey(FViewport* Viewport, int32 ControllerId, FKey Key, EInputEvent Event, float AmountDepressed = 1.0f, bool bGamepad = false) OVERRIDE;
-	virtual UWorld* GetWorld() const OVERRIDE { return World; }
+	virtual void Draw(FViewport* Viewport, FCanvas* Canvas) override;
+	virtual bool InputKey(FViewport* Viewport, int32 ControllerId, FKey Key, EInputEvent Event, float AmountDepressed = 1.0f, bool bGamepad = false) override;
+	virtual UWorld* GetWorld() const override { return World; }
 
 	/** FGCObject interface */
-	virtual void AddReferencedObjects(FReferenceCollector& Collector) OVERRIDE;
+	virtual void AddReferencedObjects(FReferenceCollector& Collector) override;
 
 	/** Modifies the checkerboard texture's data */
 	void ModifyCheckerboardTextureColors();
diff --git a/Engine/Source/Editor/TextureEditor/Private/TextureEditorToolkit.h b/Engine/Source/Editor/TextureEditor/Private/TextureEditorToolkit.h
--- a/Engine/Source/Editor/TextureEditor/Private/TextureEditorToolkit.h
+++ b/Engine/Source/Editor/TextureEditor/Private/TextureEditorToolkit.h
@@ -39,10 +39,10 @@ public:
 
 	// Begin FAssetEditorToolkit interface
 
-	virtual void RegisterTabSpawners( const TSharedRef<class FTabManager>& TabManager ) OVERRIDE;
-	virtual void UnregisterTabSpawners( const TSharedRef<class FTabManager>& TabManager ) OVERRIDE;
+	virtual void RegisterTabSpawners( const TSharedRef<class FTabManager>& TabManager ) override;
+	virtual void UnregisterTabSpawners( const TSharedRef<class FTabManager>& TabManager ) override;
 
-	virtual FString GetDocumentationLink() const OVERRIDE
+	virtual FString GetDocumentationLink() const override
 	{
 		return FString(TEXT("Engine/Content/Types/Textures/Properties/Interface"));
 	}
@@ -53,21 +53,21 @@ public:
 
 	// Begin ITextureEditorToolkit interface
 
-	virtual void CalculateTextureDimensions( uint32& Width, uint32& Height ) const OVERRIDE;
+	virtual void CalculateTextureDimensions( uint32& Width, uint32& Height ) const override;
 
-	virtual ESimpleElementBlendMode GetColourChannelBlendMode( ) const OVERRIDE;
+	virtual ESimpleElementBlendMode GetColourChannelBlendMode( ) const override;
 
-	virtual int32 GetMipLevel( ) const OVERRIDE
+	virtual int32 GetMipLevel( ) const override
 	{
 		return GetUseSpecifiedMip() ? SpecifiedMipLevel : 0;
 	}
 
-	virtual UTexture* GetTexture( ) const OVERRIDE
+	virtual UTexture* GetTexture( ) const override
 	{
 		return Texture;
 	}
 
-	virtual void PopulateQuickInfo( ) OVERRIDE;
+	virtual void PopulateQuickInfo( ) override;
 	
 	// End ITextureEditorToolkit interface
 
@@ -75,28 +75,28 @@ public:
 
 	// Begin IToolkit interface
 
-	virtual FText GetBaseToolkitName( ) const OVERRIDE;
+	virtual FText GetBaseToolkitName( ) const override;
 
-	virtual FName GetToolkitFName( ) const OVERRIDE
+	virtual FName GetToolkitFName( ) const override
 	{
 		return FName("TextureEditor");
 	}
 
-	virtual FLinearColor GetWorldCentricTabColorScale( ) const OVERRIDE
+	virtual FLinearColor GetWorldCentricTabColorScale( ) const override
 	{
 		return FLinearColor(0.3f, 0.2f, 0.5f, 0.5f);
 	}
 
-	virtual bool GetUseSpecifiedMip( ) const OVERRIDE;
+	virtual bool GetUseSpecifiedMip( ) const override;
 
-	virtual FString GetWorldCentricTabPrefix( ) const OVERRIDE;
+	virtual FString GetWorldCentricTabPrefix( ) const override;
 
-	virtual double GetZoom( ) const OVERRIDE;
-	virtual void SetZoom( double ZoomValue ) OVERRIDE;
-	virtual void ZoomIn( ) OVERRIDE;
-	virtual void ZoomOut( ) OVERRIDE;
-	virtual bool GetFitToViewport( ) const OVERRIDE;
-	virtual void SetFitToViewport( const bool bFitToViewport ) OVERRIDE;
+	virtual double GetZoom( ) const override;
+	virtual void SetZoom( double ZoomValue ) override;
+	virtual void ZoomIn( ) override;
+	virtual void ZoomOut( ) override;
+	virtual bool GetFitToViewport( ) const override;
+	virtual void SetFitToViewport( const bool bFitToViewport ) override;
 
 	// End IToolkit interface
 
@@ -104,7 +104,7 @@ public:
 
 	// Begin FGCObject interface
 
-	virtual void AddReferencedObjects( FReferenceCollector& Collector ) OVERRIDE;
+	virtual void AddReferencedObjects( FReferenceCollector& Collector ) override;
 
 	// End FGCObject interface
 	
@@ -112,9 +112,9 @@ protected:
 
 	// Begin FEditorUndoClient interface
 
-	virtual void PostUndo(bool bSuccess) OVERRIDE { }
+	virtual void PostUndo(bool bSuccess) override { }
 
-	virtual void PostRedo(bool bSuccess) OVERRIDE
+	virtual void PostRedo(bool bSuccess) override
 	{
 		PostUndo(bSuccess);
 	}
diff --git a/Engine/Source/Editor/TranslationEditor/Private/TranslationEditor.h b/Engine/Source/Editor/TranslationEditor/Private/TranslationEditor.h
--- a/Engine/Source/Editor/TranslationEditor/Private/TranslationEditor.h
+++ b/Engine/Source/Editor/TranslationEditor/Private/TranslationEditor.h
@@ -33,8 +33,8 @@ public:
 	
 	virtual ~FTranslationEditor() {}
 
-	virtual void RegisterTabSpawners(const TSharedRef<class FTabManager>& TabManager) OVERRIDE;
-	virtual void UnregisterTabSpawners(const TSharedRef<class FTabManager>& TabManager) OVERRIDE;
+	virtual void RegisterTabSpawners(const TSharedRef<class FTabManager>& TabManager) override;
+	virtual void UnregisterTabSpawners(const TSharedRef<class FTabManager>& TabManager) override;
 
 	/**
 	 * Edits the specified table
@@ -46,15 +46,15 @@ public:
 	void InitTranslationEditor( const EToolkitMode::Type Mode, const TSharedPtr< class IToolkitHost >& InitToolkitHost );
 
 	/** IToolkit interface */
-	virtual FName GetToolkitFName() const OVERRIDE;
-	virtual FText GetBaseToolkitName() const OVERRIDE;
-	virtual FText GetToolkitName() const OVERRIDE;
-	virtual FString GetWorldCentricTabPrefix() const OVERRIDE;
-	virtual FLinearColor GetWorldCentricTabColorScale() const OVERRIDE;
+	virtual FName GetToolkitFName() const override;
+	virtual FText GetBaseToolkitName() const override;
+	virtual FText GetToolkitName() const override;
+	virtual FString GetWorldCentricTabPrefix() const override;
+	virtual FLinearColor GetWorldCentricTabColorScale() const override;
 
 protected:
 	/** Called when "Save" is clicked for this asset */
-	virtual void SaveAsset_Execute() OVERRIDE;
+	virtual void SaveAsset_Execute() override;
 
 private:
 
diff --git a/Engine/Source/Editor/UMGEditor/Classes/WidgetBlueprintFactory.h b/Engine/Source/Editor/UMGEditor/Classes/WidgetBlueprintFactory.h
--- a/Engine/Source/Editor/UMGEditor/Classes/WidgetBlueprintFactory.h
+++ b/Engine/Source/Editor/UMGEditor/Classes/WidgetBlueprintFactory.h
@@ -18,8 +18,8 @@ class UWidgetBlueprintFactory : public UFactory
 	TSubclassOf<class UUserWidget> ParentClass;
 
 	// Begin UFactory Interface
-	virtual bool ConfigureProperties() OVERRIDE;
-	virtual UObject* FactoryCreateNew(UClass* Class, UObject* InParent, FName Name, EObjectFlags Flags, UObject* Context, FFeedbackContext* Warn, FName CallingContext) OVERRIDE;
-	virtual UObject* FactoryCreateNew(UClass* Class, UObject* InParent, FName Name, EObjectFlags Flags, UObject* Context, FFeedbackContext* Warn) OVERRIDE;
+	virtual bool ConfigureProperties() override;
+	virtual UObject* FactoryCreateNew(UClass* Class, UObject* InParent, FName Name, EObjectFlags Flags, UObject* Context, FFeedbackContext* Warn, FName CallingContext) override;
+	virtual UObject* FactoryCreateNew(UClass* Class, UObject* InParent, FName Name, EObjectFlags Flags, UObject* Context, FFeedbackContext* Warn) override;
 	// Begin UFactory Interface	
 };
diff --git a/Engine/Source/Editor/UMGEditor/Classes/WidgetGraphNode_Base.h b/Engine/Source/Editor/UMGEditor/Classes/WidgetGraphNode_Base.h
--- a/Engine/Source/Editor/UMGEditor/Classes/WidgetGraphNode_Base.h
+++ b/Engine/Source/Editor/UMGEditor/Classes/WidgetGraphNode_Base.h
@@ -13,25 +13,25 @@ class UWidgetGraphNode_Base : public UK2Node
 	TArray<FOptionalPinFromProperty> ShowPinForProperties;
 
 	//// UObject interface
-	UMGEDITOR_API virtual void PostEditChangeProperty(struct FPropertyChangedEvent& PropertyChangedEvent) OVERRIDE;
+	UMGEDITOR_API virtual void PostEditChangeProperty(struct FPropertyChangedEvent& PropertyChangedEvent) override;
 	//// End of UObject interface
 
 	//// UEdGraphNode interface
-	UMGEDITOR_API virtual void AllocateDefaultPins() OVERRIDE;
-	UMGEDITOR_API virtual FLinearColor GetNodeTitleColor() const OVERRIDE;
+	UMGEDITOR_API virtual void AllocateDefaultPins() override;
+	UMGEDITOR_API virtual FLinearColor GetNodeTitleColor() const override;
 	//UMGEDITOR_API virtual FString GetDocumentationLink() const OVERRIDE;
 	//UMGEDITOR_API virtual void GetPinHoverText(const UEdGraphPin& Pin, FString& HoverTextOut) const OVERRIDE;
-	UMGEDITOR_API virtual bool ShowPaletteIconOnNode() const OVERRIDE { return false; }
+	UMGEDITOR_API virtual bool ShowPaletteIconOnNode() const override { return false; }
 	//	// End of UEdGraphNode interface
 
 
 	// UK2Node interface
-	virtual bool NodeCausesStructuralBlueprintChange() const OVERRIDE { return true; }
-	virtual bool ShouldShowNodeProperties() const OVERRIDE { return true; }
-	virtual bool CanPlaceBreakpoints() const OVERRIDE { return false; }
-	UMGEDITOR_API virtual void ReallocatePinsDuringReconstruction(TArray<UEdGraphPin*>& OldPins) OVERRIDE;
-	UMGEDITOR_API virtual bool CanCreateUnderSpecifiedSchema(const UEdGraphSchema* DesiredSchema) const OVERRIDE;
-	UMGEDITOR_API virtual void GetMenuEntries(FGraphContextMenuBuilder& ContextMenuBuilder) const OVERRIDE;
+	virtual bool NodeCausesStructuralBlueprintChange() const override { return true; }
+	virtual bool ShouldShowNodeProperties() const override { return true; }
+	virtual bool CanPlaceBreakpoints() const override { return false; }
+	UMGEDITOR_API virtual void ReallocatePinsDuringReconstruction(TArray<UEdGraphPin*>& OldPins) override;
+	UMGEDITOR_API virtual bool CanCreateUnderSpecifiedSchema(const UEdGraphSchema* DesiredSchema) const override;
+	UMGEDITOR_API virtual void GetMenuEntries(FGraphContextMenuBuilder& ContextMenuBuilder) const override;
 	// End of UK2Node interface
 
 	// UWidgetGraphNode_Base interface
diff --git a/Engine/Source/Editor/UMGEditor/Private/AssetTypeActions_WidgetBlueprint.h b/Engine/Source/Editor/UMGEditor/Private/AssetTypeActions_WidgetBlueprint.h
--- a/Engine/Source/Editor/UMGEditor/Private/AssetTypeActions_WidgetBlueprint.h
+++ b/Engine/Source/Editor/UMGEditor/Private/AssetTypeActions_WidgetBlueprint.h
@@ -7,11 +7,11 @@ class FAssetTypeActions_WidgetBlueprint : public FAssetTypeActions_Base
 {
 public:
 	// IAssetTypeActions Implementation
-	virtual FText GetName() const OVERRIDE { return NSLOCTEXT("AssetTypeActions", "AssetTypeActions_WidgetBlueprint", "Widget Blueprint"); }
-	virtual FColor GetTypeColor() const OVERRIDE { return FColor(255,255,255); }
-	virtual UClass* GetSupportedClass() const OVERRIDE { return UWidgetBlueprint::StaticClass(); }
-	virtual void OpenAssetEditor(const TArray<UObject*>& InObjects, TSharedPtr<class IToolkitHost> EditWithinLevelEditor = TSharedPtr<IToolkitHost>()) OVERRIDE;
-	virtual uint32 GetCategories() OVERRIDE { return EAssetTypeCategories::UI; }
+	virtual FText GetName() const override { return NSLOCTEXT("AssetTypeActions", "AssetTypeActions_WidgetBlueprint", "Widget Blueprint"); }
+	virtual FColor GetTypeColor() const override { return FColor(255,255,255); }
+	virtual UClass* GetSupportedClass() const override { return UWidgetBlueprint::StaticClass(); }
+	virtual void OpenAssetEditor(const TArray<UObject*>& InObjects, TSharedPtr<class IToolkitHost> EditWithinLevelEditor = TSharedPtr<IToolkitHost>()) override;
+	virtual uint32 GetCategories() override { return EAssetTypeCategories::UI; }
 	// End IAssetTypeActions Implementation
 
 private:
diff --git a/Engine/Source/Editor/UMGEditor/Private/Customizations/CanvasSlotCustomization.h b/Engine/Source/Editor/UMGEditor/Private/Customizations/CanvasSlotCustomization.h
--- a/Engine/Source/Editor/UMGEditor/Private/Customizations/CanvasSlotCustomization.h
+++ b/Engine/Source/Editor/UMGEditor/Private/Customizations/CanvasSlotCustomization.h
@@ -16,9 +16,9 @@ public:
 	{
 	}
 	
-	virtual void CustomizeHeader(TSharedRef<IPropertyHandle> PropertyHandle, class FDetailWidgetRow& HeaderRow, IPropertyTypeCustomizationUtils& CustomizationUtils) OVERRIDE;
+	virtual void CustomizeHeader(TSharedRef<IPropertyHandle> PropertyHandle, class FDetailWidgetRow& HeaderRow, IPropertyTypeCustomizationUtils& CustomizationUtils) override;
 
-	virtual void CustomizeChildren(TSharedRef<IPropertyHandle> PropertyHandle, class IDetailChildrenBuilder& ChildBuilder, IPropertyTypeCustomizationUtils& CustomizationUtils) OVERRIDE;
+	virtual void CustomizeChildren(TSharedRef<IPropertyHandle> PropertyHandle, class IDetailChildrenBuilder& ChildBuilder, IPropertyTypeCustomizationUtils& CustomizationUtils) override;
 	
 private:
 
diff --git a/Engine/Source/Editor/UMGEditor/Private/Extensions/CanvasSlotExtension.h b/Engine/Source/Editor/UMGEditor/Private/Extensions/CanvasSlotExtension.h
--- a/Engine/Source/Editor/UMGEditor/Private/Extensions/CanvasSlotExtension.h
+++ b/Engine/Source/Editor/UMGEditor/Private/Extensions/CanvasSlotExtension.h
@@ -11,7 +11,7 @@ public:
 
 	bool IsActive(const TArray< FSelectedWidget >& Selection);
 	
-	virtual void BuildWidgetsForSelection(const TArray< FSelectedWidget >& Selection, TArray< TSharedRef<SWidget> >& Widgets) OVERRIDE;
+	virtual void BuildWidgetsForSelection(const TArray< FSelectedWidget >& Selection, TArray< TSharedRef<SWidget> >& Widgets) override;
 
 private:
 
diff --git a/Engine/Source/Editor/UMGEditor/Private/Extensions/HorizontalSlotExtension.h b/Engine/Source/Editor/UMGEditor/Private/Extensions/HorizontalSlotExtension.h
--- a/Engine/Source/Editor/UMGEditor/Private/Extensions/HorizontalSlotExtension.h
+++ b/Engine/Source/Editor/UMGEditor/Private/Extensions/HorizontalSlotExtension.h
@@ -11,7 +11,7 @@ public:
 
 	bool IsActive(const TArray< FSelectedWidget >& Selection);
 	
-	virtual void BuildWidgetsForSelection(const TArray< FSelectedWidget >& Selection, TArray< TSharedRef<SWidget> >& Widgets) OVERRIDE;
+	virtual void BuildWidgetsForSelection(const TArray< FSelectedWidget >& Selection, TArray< TSharedRef<SWidget> >& Widgets) override;
 
 private:
 
diff --git a/Engine/Source/Editor/UMGEditor/Private/SDesignSurface.cpp b/Engine/Source/Editor/UMGEditor/Private/SDesignSurface.cpp
--- a/Engine/Source/Editor/UMGEditor/Private/SDesignSurface.cpp
+++ b/Engine/Source/Editor/UMGEditor/Private/SDesignSurface.cpp
@@ -53,13 +53,13 @@ struct FFixedZoomLevelsContainer : public FZoomLevelsContainer
 		}
 	}
 
-	float GetZoomAmount(int32 InZoomLevel) const OVERRIDE
+	float GetZoomAmount(int32 InZoomLevel) const override
 	{
 		checkSlow(ZoomLevels.IsValidIndex(InZoomLevel));
 		return ZoomLevels[InZoomLevel].ZoomAmount;
 	}
 
-	int32 GetNearestZoomLevel(float InZoomAmount) const OVERRIDE
+	int32 GetNearestZoomLevel(float InZoomAmount) const override
 	{
 		for ( int32 ZoomLevelIndex=0; ZoomLevelIndex < GetNumZoomLevels(); ++ZoomLevelIndex )
 		{
@@ -72,23 +72,23 @@ struct FFixedZoomLevelsContainer : public FZoomLevelsContainer
 		return GetDefaultZoomLevel();
 	}
 
-	FText GetZoomText(int32 InZoomLevel) const OVERRIDE
+	FText GetZoomText(int32 InZoomLevel) const override
 	{
 		checkSlow(ZoomLevels.IsValidIndex(InZoomLevel));
 		return ZoomLevels[InZoomLevel].DisplayText;
 	}
 
-	int32 GetNumZoomLevels() const OVERRIDE
+	int32 GetNumZoomLevels() const override
 	{
 		return ZoomLevels.Num();
 	}
 
-	int32 GetDefaultZoomLevel() const OVERRIDE
+	int32 GetDefaultZoomLevel() const override
 	{
 		return 12;
 	}
 
-	EGraphRenderingLOD::Type GetLOD(int32 InZoomLevel) const OVERRIDE
+	EGraphRenderingLOD::Type GetLOD(int32 InZoomLevel) const override
 	{
 		checkSlow(ZoomLevels.IsValidIndex(InZoomLevel));
 		return ZoomLevels[InZoomLevel].LOD;
diff --git a/Engine/Source/Editor/UMGEditor/Private/SDesignSurface.h b/Engine/Source/Editor/UMGEditor/Private/SDesignSurface.h
--- a/Engine/Source/Editor/UMGEditor/Private/SDesignSurface.h
+++ b/Engine/Source/Editor/UMGEditor/Private/SDesignSurface.h
@@ -17,8 +17,8 @@ public:
 	void Construct(const FArguments& InArgs);
 
 	// SWidget interface
-	virtual FReply OnMouseWheel(const FGeometry& MyGeometry, const FPointerEvent& MouseEvent) OVERRIDE;
-	virtual int32 OnPaint(const FGeometry& AllottedGeometry, const FSlateRect& MyClippingRect, FSlateWindowElementList& OutDrawElements, int32 LayerId, const FWidgetStyle& InWidgetStyle, bool bParentEnabled) const OVERRIDE;
+	virtual FReply OnMouseWheel(const FGeometry& MyGeometry, const FPointerEvent& MouseEvent) override;
+	virtual int32 OnPaint(const FGeometry& AllottedGeometry, const FSlateRect& MyClippingRect, FSlateWindowElementList& OutDrawElements, int32 LayerId, const FWidgetStyle& InWidgetStyle, bool bParentEnabled) const override;
 	// End of Swidget interface
 
 protected:
diff --git a/Engine/Source/Editor/UMGEditor/Private/SUMGDesigner.h b/Engine/Source/Editor/UMGEditor/Private/SUMGDesigner.h
--- a/Engine/Source/Editor/UMGEditor/Private/SUMGDesigner.h
+++ b/Engine/Source/Editor/UMGEditor/Private/SUMGDesigner.h
@@ -20,19 +20,19 @@ public:
 	virtual ~SUMGDesigner();
 
 	// SWidget interface
-	virtual FReply OnMouseButtonDown(const FGeometry& MyGeometry, const FPointerEvent& MouseEvent) OVERRIDE;
-	virtual FReply OnMouseButtonUp(const FGeometry& MyGeometry, const FPointerEvent& MouseEvent) OVERRIDE;
-	virtual FReply OnMouseMove(const FGeometry& MyGeometry, const FPointerEvent& MouseEvent) OVERRIDE;
+	virtual FReply OnMouseButtonDown(const FGeometry& MyGeometry, const FPointerEvent& MouseEvent) override;
+	virtual FReply OnMouseButtonUp(const FGeometry& MyGeometry, const FPointerEvent& MouseEvent) override;
+	virtual FReply OnMouseMove(const FGeometry& MyGeometry, const FPointerEvent& MouseEvent) override;
 
-	virtual FCursorReply OnCursorQuery(const FGeometry& MyGeometry, const FPointerEvent& CursorEvent) const OVERRIDE;
+	virtual FCursorReply OnCursorQuery(const FGeometry& MyGeometry, const FPointerEvent& CursorEvent) const override;
 
-	virtual int32 OnPaint(const FGeometry& AllottedGeometry, const FSlateRect& MyClippingRect, FSlateWindowElementList& OutDrawElements, int32 LayerId, const FWidgetStyle& InWidgetStyle, bool bParentEnabled) const OVERRIDE;
-	virtual void Tick(const FGeometry& AllottedGeometry, const double InCurrentTime, const float InDeltaTime) OVERRIDE;
+	virtual int32 OnPaint(const FGeometry& AllottedGeometry, const FSlateRect& MyClippingRect, FSlateWindowElementList& OutDrawElements, int32 LayerId, const FWidgetStyle& InWidgetStyle, bool bParentEnabled) const override;
+	virtual void Tick(const FGeometry& AllottedGeometry, const double InCurrentTime, const float InDeltaTime) override;
 
-	virtual void OnDragEnter(const FGeometry& MyGeometry, const FDragDropEvent& DragDropEvent) OVERRIDE;
-	virtual void OnDragLeave(const FDragDropEvent& DragDropEvent) OVERRIDE;
-	virtual FReply OnDragOver(const FGeometry& MyGeometry, const FDragDropEvent& DragDropEvent) OVERRIDE;
-	virtual FReply OnDrop(const FGeometry& MyGeometry, const FDragDropEvent& DragDropEvent) OVERRIDE;
+	virtual void OnDragEnter(const FGeometry& MyGeometry, const FDragDropEvent& DragDropEvent) override;
+	virtual void OnDragLeave(const FDragDropEvent& DragDropEvent) override;
+	virtual FReply OnDragOver(const FGeometry& MyGeometry, const FDragDropEvent& DragDropEvent) override;
+	virtual FReply OnDrop(const FGeometry& MyGeometry, const FDragDropEvent& DragDropEvent) override;
 	// End of SWidget interface
 
 	void Register(TSharedRef<FDesignerExtension> Extension);
