Static Mesh Editor: Ctrl+W now duplicates correctly.

#ttp 338927 - EDITOR: Static Mesh: Unable to duplicate collision primitives with hotkey

#branch UE4

#change W input was getting swallowed by the input to switch between widget maninpulation types - converted this over to use the UICommands that are setup

diff --git a/Engine/Source/Editor/StaticMeshEditor/Private/StaticMeshEditorViewportClient.cpp b/Engine/Source/Editor/StaticMeshEditor/Private/StaticMeshEditorViewportClient.cpp
--- a/Engine/Source/Editor/StaticMeshEditor/Private/StaticMeshEditorViewportClient.cpp
+++ b/Engine/Source/Editor/StaticMeshEditor/Private/StaticMeshEditorViewportClient.cpp
@@ -292,6 +292,47 @@ FWidget::EWidgetMode FStaticMeshEditorViewportClient::GetWidgetMode() const
 	return FWidget::WM_None;
 }
 
+void FStaticMeshEditorViewportClient::SetWidgetMode(FWidget::EWidgetMode NewMode)
+{
+	WidgetMode = NewMode;
+	Invalidate();
+}
+
+bool FStaticMeshEditorViewportClient::CanSetWidgetMode(FWidget::EWidgetMode NewMode) const
+{
+	if (!Widget->IsDragging())
+	{
+		const bool bSelectedPrim = StaticMeshEditorPtr.Pin()->HasSelectedPrims();
+		if (bSelectedPrim)
+		{
+			return true;
+		}
+		else if (NewMode != FWidget::WM_Scale)	// Sockets don't support scaling
+		{
+			const UStaticMeshSocket* SelectedSocket = StaticMeshEditorPtr.Pin()->GetSelectedSocket();
+			if (SelectedSocket)
+			{
+				return true;
+			}
+		}
+	}
+	return false;
+}
+
+bool FStaticMeshEditorViewportClient::CanCycleWidgetMode() const
+{
+	if (!Widget->IsDragging())
+	{
+		const UStaticMeshSocket* SelectedSocket = StaticMeshEditorPtr.Pin()->GetSelectedSocket();
+		const bool bSelectedPrim = StaticMeshEditorPtr.Pin()->HasSelectedPrims();
+		if ((SelectedSocket || bSelectedPrim))
+		{
+			return true;
+		}
+	}
+	return false;
+}
+
 void FStaticMeshEditorViewportClient::TrackingStopped()
 {
 	if( bManipulating )
@@ -768,98 +809,7 @@ void FStaticMeshEditorViewportClient::MouseMove(FViewport* Viewport,int32 x, int
 
 bool FStaticMeshEditorViewportClient::InputKey(FViewport* Viewport, int32 ControllerId, FKey Key, EInputEvent Event,float AmountDepressed,bool Gamepad)
 {
-	bool bHandled = false;
-
-	const int32 HitX = Viewport->GetMouseX();
-	const int32 HitY = Viewport->GetMouseY();
-	const bool bAltDown = Viewport->KeyState(EKeys::LeftAlt) || Viewport->KeyState(EKeys::RightAlt);
-	const bool bCtrlDown = Viewport->KeyState(EKeys::LeftControl) || Viewport->KeyState(EKeys::RightControl);
-	const bool bShiftDown = Viewport->KeyState(EKeys::LeftShift) || Viewport->KeyState(EKeys::RightShift);
-
-	if( Event == IE_Pressed) 
-	{
-		if (Key == EKeys::W)
-		{
-			if (!Widget->IsDragging())
-			{
-				const UStaticMeshSocket* SelectedSocket = StaticMeshEditorPtr.Pin()->GetSelectedSocket();
-				const bool bSelectedPrim = StaticMeshEditorPtr.Pin()->HasSelectedPrims();
-				if ((SelectedSocket || bSelectedPrim))
-				{
-					WidgetMode = FWidget::WM_Translate;
-					Invalidate();
-				}
-			}
-			bHandled = true;
-		}
-		else if (Key == EKeys::E)
-		{
-			if (!Widget->IsDragging())
-			{
-				const UStaticMeshSocket* SelectedSocket = StaticMeshEditorPtr.Pin()->GetSelectedSocket();
-				const bool bSelectedPrim = StaticMeshEditorPtr.Pin()->HasSelectedPrims();
-				if ((SelectedSocket || bSelectedPrim))
-				{
-					WidgetMode = FWidget::WM_Rotate;
-					Invalidate();
-				}
-			}
-			bHandled = true;
-		}
-		else if (Key == EKeys::R)
-		{
-			if (!Widget->IsDragging())
-			{
-				const bool bSelectedPrim = StaticMeshEditorPtr.Pin()->HasSelectedPrims();
-				if ((bSelectedPrim))	// Sockets don't support scaling
-				{
-					WidgetMode = FWidget::WM_Scale;
-					Invalidate();
-				}
-			}
-			bHandled = true;
-		}
-		else if(Key == EKeys::SpaceBar)
-		{
-			if(!Widget->IsDragging())
-			{
-				const UStaticMeshSocket* SelectedSocket = StaticMeshEditorPtr.Pin()->GetSelectedSocket();
-				const bool bSelectedPrim = StaticMeshEditorPtr.Pin()->HasSelectedPrims();
-				if ((SelectedSocket || bSelectedPrim))
-				{
-					const FWidget::EWidgetMode MoveMode = GetWidgetMode();
-					if(MoveMode == FWidget::WM_Rotate)
-					{
-						// Sockets don't support scaling
-						if (bSelectedPrim)
-						{
-							WidgetMode = FWidget::WM_Scale;
-						}
-						else
-						{
-							WidgetMode = FWidget::WM_Translate;
-						}
-					}
-					else if(MoveMode == FWidget::WM_Scale)
-					{
-						WidgetMode = FWidget::WM_Translate;
-					}
-					else if(MoveMode == FWidget::WM_Translate)
-					{
-						WidgetMode = FWidget::WM_Rotate;
-					}
-					Invalidate();
-				}
-			}
-			bHandled = true;
-		}
-	}
-
-
-	if( !bHandled )
-	{
-		bHandled = FEditorViewportClient::InputKey( Viewport, ControllerId, Key, Event, AmountDepressed, false );
-	}
+	bool bHandled = FEditorViewportClient::InputKey(Viewport, ControllerId, Key, Event, AmountDepressed, false);
 
 	// Handle viewport screenshot.
 	bHandled |= InputTakeScreenshot( Viewport, Key, Event );
diff --git a/Engine/Source/Editor/StaticMeshEditor/Private/StaticMeshEditorViewportClient.h b/Engine/Source/Editor/StaticMeshEditor/Private/StaticMeshEditorViewportClient.h
--- a/Engine/Source/Editor/StaticMeshEditor/Private/StaticMeshEditorViewportClient.h
+++ b/Engine/Source/Editor/StaticMeshEditor/Private/StaticMeshEditorViewportClient.h
@@ -22,6 +22,9 @@ public:
 	virtual void TrackingStarted( const struct FInputEventState& InInputState, bool bIsDragging, bool bNudge ) override;
 	virtual void TrackingStopped() override;
 	virtual FWidget::EWidgetMode GetWidgetMode() const override;
+	virtual void SetWidgetMode(FWidget::EWidgetMode NewMode) override;
+	virtual bool CanSetWidgetMode(FWidget::EWidgetMode NewMode) const override;
+	virtual bool CanCycleWidgetMode() const override;
 	virtual FVector GetWidgetLocation() const;
 	virtual FMatrix GetWidgetCoordSystem() const;
 	virtual ECoordSystem GetWidgetCoordSystemSpace() const { return COORD_Local; }
