Moved static mesh stat info text into Slate UI.
#ttp 339600 - Move static mesh stat info out of canvas and into StaticMeshEditor Slate UI
#branch UE4
#proj Editor.EditorStyle, Editor.StaticMeshEditor
#add Added a parameter to the StaticMeshEditorViewportClient to pass in a weak ptr to the viewport widget which owns it.  This is so that, on Draw calls, we can update one of its child Slate widgets.
#add Created a struct SStaticMeshEditorViewport::FOverlayTextItem which represents one item of overlay text, and a method SStaticMeshEditorViewport::PopulateOverlayText() which populates the Slate SVerticalBox in the SOverlay with these text blocks.
#add Created a couple of new SlateEditorStyles - shadowed text (normal), and shadowed text (error).
#change Changed FStaticMeshEditorViewportClient::DrawCanvas() to populate an FOverlayTextItem array and call PopulateOverlayText on the owning widget to update the overlay widget.

diff --git a/Engine/Source/Editor/EditorStyle/Private/SlateEditorStyle.cpp b/Engine/Source/Editor/EditorStyle/Private/SlateEditorStyle.cpp
--- a/Engine/Source/Editor/EditorStyle/Private/SlateEditorStyle.cpp
+++ b/Engine/Source/Editor/EditorStyle/Private/SlateEditorStyle.cpp
@@ -181,6 +181,15 @@ void FSlateEditorStyle::FStyle::SetupGeneralStyles()
 		Set( "TextBlock.HighlightShape",  new BOX_BRUSH( "Common/TextBlockHighlightShape", FMargin(3.f/8.f) ));
 		Set( "TextBlock.HighlighColor", FLinearColor( 0.02f, 0.3f, 0.0f ) );
 
+		Set("TextBlock.ShadowedText", FTextBlockStyle(NormalText)
+			.SetShadowOffset(FVector2D(1.0f, 1.0f))
+			.SetShadowColorAndOpacity(FLinearColor(0.0f, 0.0f, 0.0f)));
+
+		Set("TextBlock.ShadowedTextWarning", FTextBlockStyle(NormalText)
+			.SetColorAndOpacity(FLinearColor(1.0f, 0.0f, 0.0f))
+			.SetShadowOffset(FVector2D(1.0f, 1.0f))
+			.SetShadowColorAndOpacity(FLinearColor(0.0f, 0.0f, 0.0f)));
+
 		Set("NormalText", NormalText);
 	}
 
diff --git a/Engine/Source/Editor/StaticMeshEditor/Private/SStaticMeshEditorViewport.cpp b/Engine/Source/Editor/StaticMeshEditor/Private/SStaticMeshEditorViewport.cpp
--- a/Engine/Source/Editor/StaticMeshEditor/Private/SStaticMeshEditorViewport.cpp
+++ b/Engine/Source/Editor/StaticMeshEditor/Private/SStaticMeshEditorViewport.cpp
@@ -44,6 +44,14 @@ void SStaticMeshEditorViewport::Construct(const FArguments& InArgs)
 
 	SetPreviewMesh(StaticMesh);
 
+	ViewportOverlay->AddSlot()
+		.VAlign(VAlign_Top)
+		.HAlign(HAlign_Left)
+		.Padding(10)
+		[
+			SAssignNew(OverlayTextVerticalBox, SVerticalBox)
+		];
+
 	FCoreDelegates::OnObjectPropertyChanged.AddRaw(this, &SStaticMeshEditorViewport::OnObjectPropertyChanged);
 
 }
@@ -57,6 +65,21 @@ SStaticMeshEditorViewport::~SStaticMeshEditorViewport()
 	}
 }
 
+void SStaticMeshEditorViewport::PopulateOverlayText(const TArray<FOverlayTextItem>& TextItems)
+{
+	OverlayTextVerticalBox->ClearChildren();
+
+	for (const auto& TextItem : TextItems)
+	{
+		OverlayTextVerticalBox->AddSlot()
+		[
+			SNew(STextBlock)
+			.Text(TextItem.Text)
+			.TextStyle(FEditorStyle::Get(), TextItem.Style)
+		];
+	}
+}
+
 void SStaticMeshEditorViewport::AddReferencedObjects( FReferenceCollector& Collector )
 {
 	Collector.AddReferencedObject( PreviewMeshComponent );
@@ -273,7 +296,7 @@ FStaticMeshEditorViewportClient& SStaticMeshEditorViewport::GetViewportClient()
 
 TSharedRef<FEditorViewportClient> SStaticMeshEditorViewport::MakeEditorViewportClient()
 {
-	EditorViewportClient = MakeShareable( new FStaticMeshEditorViewportClient(StaticMeshEditorPtr, PreviewScene, StaticMesh, NULL) );
+	EditorViewportClient = MakeShareable( new FStaticMeshEditorViewportClient(StaticMeshEditorPtr, SharedThis(this), PreviewScene, StaticMesh, NULL) );
 
 	EditorViewportClient->bSetListenerPosition = false;
 
diff --git a/Engine/Source/Editor/StaticMeshEditor/Private/SStaticMeshEditorViewport.h b/Engine/Source/Editor/StaticMeshEditor/Private/SStaticMeshEditorViewport.h
--- a/Engine/Source/Editor/StaticMeshEditor/Private/SStaticMeshEditorViewport.h
+++ b/Engine/Source/Editor/StaticMeshEditor/Private/SStaticMeshEditorViewport.h
@@ -67,6 +67,20 @@ public:
 	/** Set the parent tab of the viewport for determining visibility */
 	void SetParentTab( TSharedRef<SDockTab> InParentTab ) { ParentTab = InParentTab; }
 
+	/** Struct defining the text and its style of each item in the overlay widget */
+	struct FOverlayTextItem
+	{
+		explicit FOverlayTextItem(const FText& InText, const FName& InStyle = "TextBlock.ShadowedText")
+			: Text(InText), Style(InStyle)
+		{}
+
+		FText Text;
+		FName Style;
+	};
+
+	/** Specifies an array of text items which will be added to the viewport overlay */
+	void PopulateOverlayText( const TArray<FOverlayTextItem>& TextItems );
+
 protected:
 	/** SEditorViewport interface */
 	virtual TSharedRef<FEditorViewportClient> MakeEditorViewportClient() override;
@@ -74,6 +88,7 @@ protected:
 	virtual void BindCommands() override;
 	virtual void OnFocusViewportToSelection() override;
 	virtual TSharedPtr<SWidget> MakeViewportToolbar() override;
+
 private:
 	/** Determines the visibility of the viewport. */
 	bool IsVisible() const;
@@ -114,4 +129,7 @@ private:
 
 	/** The currently selected view mode. */
 	EViewModeIndex CurrentViewMode;
+
+	/** Pointer to the vertical box into which the overlay text items are added */
+	TSharedPtr<SVerticalBox> OverlayTextVerticalBox;
 };
diff --git a/Engine/Source/Editor/StaticMeshEditor/Private/StaticMeshEditorViewportClient.cpp b/Engine/Source/Editor/StaticMeshEditor/Private/StaticMeshEditorViewportClient.cpp
--- a/Engine/Source/Editor/StaticMeshEditor/Private/StaticMeshEditorViewportClient.cpp
+++ b/Engine/Source/Editor/StaticMeshEditor/Private/StaticMeshEditorViewportClient.cpp
@@ -15,6 +15,7 @@
 #include "MeshBuild.h"
 #include "PreviewScene.h"
 #include "ObjectTools.h"
+#include "SStaticMeshEditorViewport.h"
 
 #if WITH_PHYSX
 #include "Editor/UnrealEd/Private/EditorPhysXSupport.h"
@@ -35,9 +36,10 @@ namespace {
 	static float AmbientCubemapIntensity = 0.4f;
 }
 
-FStaticMeshEditorViewportClient::FStaticMeshEditorViewportClient(TWeakPtr<IStaticMeshEditor> InStaticMeshEditor, FPreviewScene& InPreviewScene, UStaticMesh* InPreviewStaticMesh, UStaticMeshComponent* InPreviewStaticMeshComponent)
+FStaticMeshEditorViewportClient::FStaticMeshEditorViewportClient(TWeakPtr<IStaticMeshEditor> InStaticMeshEditor, TWeakPtr<SStaticMeshEditorViewport> InStaticMeshEditorViewport, FPreviewScene& InPreviewScene, UStaticMesh* InPreviewStaticMesh, UStaticMeshComponent* InPreviewStaticMeshComponent)
 	: FEditorViewportClient(GLevelEditorModeTools(), &InPreviewScene)
 	, StaticMeshEditorPtr(InStaticMeshEditor)
+	, StaticMeshEditorViewportPtr(InStaticMeshEditorViewport)
 {
 	SimplygonLogo = LoadObject<UTexture2D>(NULL, TEXT("/Engine/EditorResources/SimplygonLogo.SimplygonLogo"), NULL, LOAD_None, NULL);
 
@@ -623,6 +625,12 @@ void FStaticMeshEditorViewportClient::DrawCanvas( FViewport& InViewport, FSceneV
 	}
 #endif // #if TODO_STATICMESH
 
+	auto StaticMeshEditor = StaticMeshEditorPtr.Pin();
+	auto StaticMeshEditorViewport = StaticMeshEditorViewportPtr.Pin();
+	if (!StaticMeshEditor.IsValid() || !StaticMeshEditorViewport.IsValid())
+	{
+		return;
+	}
 
 	const int32 HalfX = Viewport->GetSizeXY().X/2;
 	const int32 HalfY = Viewport->GetSizeXY().Y/2;
@@ -647,7 +655,7 @@ void FStaticMeshEditorViewportClient::DrawCanvas( FViewport& InViewport, FSceneV
 					FCanvasTextItem TextItem( FVector2D( XPos, YPos ), FText::FromString( Socket->SocketName.ToString() ), GEngine->GetSmallFont(), FLinearColor(FColor(255,196,196)) );
 					Canvas.DrawItem( TextItem );	
 
-					const UStaticMeshSocket* SelectedSocket = StaticMeshEditorPtr.Pin()->GetSelectedSocket();
+					const UStaticMeshSocket* SelectedSocket = StaticMeshEditor->GetSelectedSocket();
 					if (bManipulating && SelectedSocket == Socket)
 					{
 						//Figure out the text height
@@ -666,8 +674,9 @@ void FStaticMeshEditorViewportClient::DrawCanvas( FViewport& InViewport, FSceneV
 		}
 	}
 
-	int32 YPos = 6;
-	int32 CurrentLODLevel = StaticMeshEditorPtr.Pin()->GetCurrentLODLevel();
+	TArray<SStaticMeshEditorViewport::FOverlayTextItem> TextItems;
+
+	int32 CurrentLODLevel = StaticMeshEditor->GetCurrentLODLevel();
 	if (CurrentLODLevel == 0)
 	{
 		CurrentLODLevel = ComputeStaticMeshLOD(StaticMesh->RenderData, StaticMeshComponent->Bounds.Origin, StaticMeshComponent->Bounds.SphereRadius, View);
@@ -677,55 +686,25 @@ void FStaticMeshEditorViewportClient::DrawCanvas( FViewport& InViewport, FSceneV
 		CurrentLODLevel -= 1;
 	}
 
-	Canvas.DrawShadowedString(
-		6,
-		YPos,
-		*FText::Format( NSLOCTEXT("UnrealEd", "LOD_F", "LOD:  {0}"), FText::AsNumber(CurrentLODLevel) ).ToString(),
-		GEngine->GetSmallFont(),
-		FLinearColor::White
-		);
-	YPos += 18;
+	TextItems.Add(SStaticMeshEditorViewport::FOverlayTextItem(
+		FText::Format(NSLOCTEXT("UnrealEd", "LOD_F", "LOD:  {0}"), FText::AsNumber(CurrentLODLevel))));
 
 	float CurrentScreenSize = ComputeBoundsScreenSize(StaticMeshComponent->Bounds.Origin, StaticMeshComponent->Bounds.SphereRadius, View);
 	FNumberFormattingOptions FormatOptions;
 	FormatOptions.MinimumFractionalDigits = 3;
 	FormatOptions.MaximumFractionalDigits = 6;
 	FormatOptions.MaximumIntegralDigits = 6;
-	Canvas.DrawShadowedString(
-		6,
-		YPos,
-		*FText::Format( NSLOCTEXT("UnrealEd", "ScreenSize_F", "Current Screen Size:  {0}"), FText::AsNumber(CurrentScreenSize, &FormatOptions)).ToString(),
-		GEngine->GetSmallFont(),
-		FLinearColor::White
-		);
-	YPos += 18;
-
-	Canvas.DrawShadowedString(
-		6,
-		YPos,
-		*FText::Format( NSLOCTEXT("UnrealEd", "Triangles_F", "Triangles:  {0}"), FText::AsNumber(StaticMeshEditorPtr.Pin()->GetNumTriangles(CurrentLODLevel)) ).ToString(),
-		GEngine->GetSmallFont(),
-		FLinearColor::White
-		);
-	YPos += 18;
-
-	Canvas.DrawShadowedString(
-		6,
-		YPos,
-		*FText::Format( NSLOCTEXT("UnrealEd", "Vertices_F", "Vertices:  {0}"), FText::AsNumber(StaticMeshEditorPtr.Pin()->GetNumVertices(CurrentLODLevel)) ).ToString(),
-		GEngine->GetSmallFont(),
-		FLinearColor::White
-		);
-	YPos += 18;
-
-	Canvas.DrawShadowedString(
-		6,
-		YPos,
-		*FText::Format( NSLOCTEXT("UnrealEd", "UVChannels_F", "UV Channels:  {0}"), FText::AsNumber(StaticMeshEditorPtr.Pin()->GetNumUVChannels(CurrentLODLevel)) ).ToString(),
-		GEngine->GetSmallFont(),
-		FLinearColor::White
-		);
-	YPos += 18;
+	TextItems.Add(SStaticMeshEditorViewport::FOverlayTextItem(
+		FText::Format(NSLOCTEXT("UnrealEd", "ScreenSize_F", "Current Screen Size:  {0}"), FText::AsNumber(CurrentScreenSize, &FormatOptions))));
+
+	TextItems.Add(SStaticMeshEditorViewport::FOverlayTextItem(
+		FText::Format(NSLOCTEXT("UnrealEd", "Triangles_F", "Triangles:  {0}"), FText::AsNumber(StaticMeshEditorPtr.Pin()->GetNumTriangles(CurrentLODLevel)))));
+
+	TextItems.Add(SStaticMeshEditorViewport::FOverlayTextItem(
+		FText::Format(NSLOCTEXT("UnrealEd", "Vertices_F", "Vertices:  {0}"), FText::AsNumber(StaticMeshEditorPtr.Pin()->GetNumVertices(CurrentLODLevel)))));
+
+	TextItems.Add(SStaticMeshEditorViewport::FOverlayTextItem(
+		FText::Format(NSLOCTEXT("UnrealEd", "UVChannels_F", "UV Channels:  {0}"), FText::AsNumber(StaticMeshEditorPtr.Pin()->GetNumUVChannels(CurrentLODLevel)))));
 
 	const FDistanceFieldVolumeData& VolumeData = StaticMesh->RenderData->LODResources[0].DistanceFieldData;
 
@@ -739,55 +718,34 @@ void FStaticMeshEditorViewportClient::DrawCanvas( FViewport& InViewport, FSceneV
 
 		if (VolumeData.bMeshWasClosed)
 		{
-			Canvas.DrawShadowedString(
-				6,
-				YPos,
-				*FText::Format( NSLOCTEXT("UnrealEd", "DistanceFieldRes_F", "Distance Field:  {0}x{1}x{2} = {3}mb"), FText::AsNumber(VolumeData.Size.X), FText::AsNumber(VolumeData.Size.Y), FText::AsNumber(VolumeData.Size.Z), FText::AsNumber(MemoryMb, &NumberOptions) ).ToString(),
-				GEngine->GetSmallFont(),
-				FLinearColor::White
-				);
+			TextItems.Add(SStaticMeshEditorViewport::FOverlayTextItem(
+				FText::Format(NSLOCTEXT("UnrealEd", "DistanceFieldRes_F", "Distance Field:  {0}x{1}x{2} = {3}Mb"), FText::AsNumber(VolumeData.Size.X), FText::AsNumber(VolumeData.Size.Y), FText::AsNumber(VolumeData.Size.Z), FText::AsNumber(MemoryMb, &NumberOptions))));
 		}
 		else
 		{
-			Canvas.DrawShadowedString(
-				6,
-				YPos,
-				*NSLOCTEXT("UnrealEd", "DistanceFieldClosed_F", "Distance Field:  Mesh was not closed and material was one-sided").ToString(),
-				GEngine->GetSmallFont(),
-				FLinearColor::Red
-				);
+			TextItems.Add(SStaticMeshEditorViewport::FOverlayTextItem(
+				NSLOCTEXT("UnrealEd", "DistanceFieldClosed_F", "Distance Field:  Mesh was not closed and material was one-sided")));
 		}
-		
-		YPos += 18;
 	}
 
-	Canvas.DrawShadowedString(
-		6,
-		YPos,
-		*FText::Format( NSLOCTEXT("UnrealEd", "ApproxSize_F", "Approx Size: {0}x{1}x{2}"),
-			FText::AsNumber(int32(StaticMesh->GetBounds().BoxExtent.X * 2.0f)), // x2 as artists wanted length not radius
-			FText::AsNumber(int32(StaticMesh->GetBounds().BoxExtent.Y * 2.0f)),
-			FText::AsNumber(int32(StaticMesh->GetBounds().BoxExtent.Z * 2.0f)) ).ToString(),
-		GEngine->GetSmallFont(),
-		FLinearColor::White
-		);
-	YPos += 18;
+	TextItems.Add(SStaticMeshEditorViewport::FOverlayTextItem(
+		FText::Format(NSLOCTEXT("UnrealEd", "ApproxSize_F", "Approx Size: {0}x{1}x{2}"),
+		FText::AsNumber(int32(StaticMesh->GetBounds().BoxExtent.X * 2.0f)), // x2 as artists wanted length not radius
+		FText::AsNumber(int32(StaticMesh->GetBounds().BoxExtent.Y * 2.0f)),
+		FText::AsNumber(int32(StaticMesh->GetBounds().BoxExtent.Z * 2.0f)))));
 
 	// Show the number of collision primitives if we are drawing collision.
 	if(bShowCollision && StaticMesh->BodySetup)
 	{
-		Canvas.DrawShadowedString(
-			6,
-			YPos,
-			*FText::Format( NSLOCTEXT("UnrealEd", "NumPrimitives_F", "Num Primitives:  {0}"), FText::AsNumber(StaticMesh->BodySetup->AggGeom.GetElementCount()) ).ToString(),
-			GEngine->GetSmallFont(),
-			FLinearColor::White
-			);
-		YPos += 18;
+		TextItems.Add(SStaticMeshEditorViewport::FOverlayTextItem(
+			FText::Format(NSLOCTEXT("UnrealEd", "NumPrimitives_F", "Num Primitives:  {0}"), FText::AsNumber(StaticMesh->BodySetup->AggGeom.GetElementCount()))));
 	}
 
+	StaticMeshEditorViewport->PopulateOverlayText(TextItems);
+
 	if(bDrawUVs)
 	{
+		const int32 YPos = 160;
 		DrawUVsForMesh(Viewport, &Canvas, YPos);
 	}
 }
diff --git a/Engine/Source/Editor/StaticMeshEditor/Private/StaticMeshEditorViewportClient.h b/Engine/Source/Editor/StaticMeshEditor/Private/StaticMeshEditorViewportClient.h
--- a/Engine/Source/Editor/StaticMeshEditor/Private/StaticMeshEditorViewportClient.h
+++ b/Engine/Source/Editor/StaticMeshEditor/Private/StaticMeshEditorViewportClient.h
@@ -8,7 +8,7 @@
 class FStaticMeshEditorViewportClient : public FEditorViewportClient, public TSharedFromThis<FStaticMeshEditorViewportClient>
 {
 public:
-	FStaticMeshEditorViewportClient(TWeakPtr<IStaticMeshEditor> InStaticMeshEditor, FPreviewScene& InPreviewScene, UStaticMesh* InPreviewStaticMesh, UStaticMeshComponent* InPreviewStaticMeshComponent);
+	FStaticMeshEditorViewportClient(TWeakPtr<IStaticMeshEditor> InStaticMeshEditor, TWeakPtr<SStaticMeshEditorViewport> InStaticMeshEditorViewport, FPreviewScene& InPreviewScene, UStaticMesh* InPreviewStaticMesh, UStaticMeshComponent* InPreviewStaticMeshComponent);
 
 	// FEditorViewportClient interface
 	virtual void MouseMove(FViewport* Viewport,int32 x, int32 y) override;
@@ -153,4 +153,7 @@ private:
 
 	/** Cached tex coords for the currently selected edges. Used for rendering UV coordinates*/
 	TArray<FVector2D> SelectedEdgeTexCoords[MAX_STATIC_TEXCOORDS];
+
+	/** Pointer back to the StaticMeshEditor viewport control that owns us */
+	TWeakPtr<SStaticMeshEditorViewport> StaticMeshEditorViewportPtr;
 };
\ No newline at end of file
