non unity fixes

diff --git a/Engine/Source/Editor/Persona/Private/PersonaMeshDetails.h b/Engine/Source/Editor/Persona/Private/PersonaMeshDetails.h
--- a/Engine/Source/Editor/Persona/Private/PersonaMeshDetails.h
+++ b/Engine/Source/Editor/Persona/Private/PersonaMeshDetails.h
@@ -2,6 +2,7 @@
 
 #pragma once
 
+#include "IDetailCustomNodeBuilder.h"
 /*
  * Struct to uniquely identify clothing applied to a material section
  * Contains index into the ClothingAssets array and the submesh index.
diff --git a/Engine/Source/Editor/StaticMeshEditor/Private/StaticMeshEditorViewportClient.h b/Engine/Source/Editor/StaticMeshEditor/Private/StaticMeshEditorViewportClient.h
--- a/Engine/Source/Editor/StaticMeshEditor/Private/StaticMeshEditorViewportClient.h
+++ b/Engine/Source/Editor/StaticMeshEditor/Private/StaticMeshEditorViewportClient.h
@@ -4,6 +4,8 @@
 
 #include "EditorViewportClient.h"
 
+class SStaticMeshEditorViewport;
+
 /** Viewport Client for the preview viewport */
 class FStaticMeshEditorViewportClient : public FEditorViewportClient, public TSharedFromThis<FStaticMeshEditorViewportClient>
 {
