Modifying the existing static mesh import options enables a button which allows the user to reimport

#ttp 340526 - Static mesh editor: changing the static meshes normals setting from "import normals" to "compute normals" doesn't change the mesh until the user reimports the asset

#branch UE4

#add Added button to reimport based on the current import options (no button will appear if there are no options to modify)

#change Modified import options so they get marked as dirty until they are reimported (button will be unclickable unless they are)

diff --git a/Engine/Plugins/Editor/SpeedTreeImporter/Source/SpeedTreeImporter/Private/SpeedTreeImportFactory.cpp b/Engine/Plugins/Editor/SpeedTreeImporter/Source/SpeedTreeImporter/Private/SpeedTreeImportFactory.cpp
--- a/Engine/Plugins/Editor/SpeedTreeImporter/Source/SpeedTreeImporter/Private/SpeedTreeImportFactory.cpp
+++ b/Engine/Plugins/Editor/SpeedTreeImporter/Source/SpeedTreeImporter/Private/SpeedTreeImportFactory.cpp
@@ -996,6 +996,7 @@ UObject* USpeedTreeImportFactory::FactoryCreateBinary(UClass* InClass, UObject*
 				StaticMesh->AssetImportData = ConstructObject<UAssetImportData>(UAssetImportData::StaticClass(), StaticMesh);
 				StaticMesh->AssetImportData->SourceFilePath = FReimportManager::SanitizeImportFilename(UFactory::GetCurrentFilename(), StaticMesh);
 				StaticMesh->AssetImportData->SourceFileTimestamp = IFileManager::Get().GetTimeStamp(*UFactory::GetCurrentFilename()).ToString();
+				StaticMesh->AssetImportData->bDirty = false;
 				
 				// clear out any old data
 				StaticMesh->SourceModels.Empty();
diff --git a/Engine/Source/Editor/StaticMeshEditor/Private/StaticMeshEditorTools.cpp b/Engine/Source/Editor/StaticMeshEditor/Private/StaticMeshEditorTools.cpp
--- a/Engine/Source/Editor/StaticMeshEditor/Private/StaticMeshEditorTools.cpp
+++ b/Engine/Source/Editor/StaticMeshEditor/Private/StaticMeshEditorTools.cpp
@@ -34,6 +34,8 @@ void FStaticMeshDetails::CustomizeDetails( class IDetailLayoutBuilder& DetailBui
 {
 	IDetailCategoryBuilder& LODSettingsCategory = DetailBuilder.EditCategory( "LodSettings", LOCTEXT("LodSettingsCategory", "LOD Settings").ToString() );
 	IDetailCategoryBuilder& StaticMeshCategory = DetailBuilder.EditCategory( "StaticMesh", LOCTEXT("StaticMeshGeneralSettings", "Static Mesh Settings").ToString() );
+	IDetailCategoryBuilder& ImportCategory = DetailBuilder.EditCategory( "ImportSettings", LOCTEXT("ImportGeneralSettings", "Import Settings").ToString() );
+	
 	DetailBuilder.EditCategory( "Navigation", TEXT(""), ECategoryPriority::Uncommon );
 
 	LevelOfDetailSettings = MakeShareable( new FLevelOfDetailSettingsLayout( StaticMeshEditor ) );
@@ -77,6 +79,26 @@ void FStaticMeshDetails::CustomizeDetails( class IDetailLayoutBuilder& DetailBui
 			}
 		}
 	}
+
+	// Only add the reimport button if we have reimport settings we can modify
+	// Note: this will get rebuilt if the asset is reimported so we don't need to use .Visibility on the button
+	const UStaticMesh* StaticMesh = StaticMeshEditor.GetStaticMesh();
+	check(StaticMesh);
+	if (StaticMesh->AssetImportData && StaticMesh->AssetImportData->GetClass() != UAssetImportData::StaticClass())
+	{
+		ImportCategory.AddCustomRow(LOCTEXT("ReimportStaticMesh", "Reimport Static Mesh").ToString(), true)
+			.ValueContent()
+			[
+				SNew(SButton)
+				.OnClicked(this, &FStaticMeshDetails::Reimport)
+				.IsEnabled(this, &FStaticMeshDetails::CanReimport)
+				[
+					SNew(STextBlock)
+					.Text(LOCTEXT("Reimport", "Reimport"))
+					.Font(IDetailLayoutBuilder::GetDetailFont())
+				]
+			];
+	}
 }
 
 void SConvexDecomposition::Construct(const FArguments& InArgs)
@@ -178,6 +200,25 @@ void FStaticMeshDetails::ApplyChanges()
 	}
 }
 
+FReply FStaticMeshDetails::Reimport()
+{
+	UStaticMesh* StaticMesh = StaticMeshEditor.GetStaticMesh();
+	check(StaticMesh);
+	FReimportManager::Instance()->Reimport(StaticMesh, true);
+	return FReply::Handled();
+}
+
+bool FStaticMeshDetails::CanReimport() const
+{
+	const UStaticMesh* StaticMesh = StaticMeshEditor.GetStaticMesh();
+	check(StaticMesh);
+	if (StaticMesh->AssetImportData)
+	{
+		return StaticMesh->AssetImportData->bDirty;
+	}
+	return false;
+}
+
 SConvexDecomposition::~SConvexDecomposition()
 {
 
diff --git a/Engine/Source/Editor/StaticMeshEditor/Private/StaticMeshEditorTools.h b/Engine/Source/Editor/StaticMeshEditor/Private/StaticMeshEditorTools.h
--- a/Engine/Source/Editor/StaticMeshEditor/Private/StaticMeshEditorTools.h
+++ b/Engine/Source/Editor/StaticMeshEditor/Private/StaticMeshEditorTools.h
@@ -33,6 +33,12 @@ public:
 
 	/** Applies level of detail changes to the static mesh */
 	void ApplyChanges();
+
+	/** Reimports the current static mesh */
+	FReply Reimport();
+
+	/** Whether the user should reimport based on the changes */
+	bool CanReimport() const;
 private:
 	/** Level of detail settings for the details panel */
 	TSharedPtr<FLevelOfDetailSettingsLayout> LevelOfDetailSettings;
diff --git a/Engine/Source/Editor/UnrealEd/Private/Factories/ApexDestructibleAssetImport.cpp b/Engine/Source/Editor/UnrealEd/Private/Factories/ApexDestructibleAssetImport.cpp
--- a/Engine/Source/Editor/UnrealEd/Private/Factories/ApexDestructibleAssetImport.cpp
+++ b/Engine/Source/Editor/UnrealEd/Private/Factories/ApexDestructibleAssetImport.cpp
@@ -1012,6 +1012,7 @@ UDestructibleMesh* ImportDestructibleMeshFromApexDestructibleAsset(UObject* InPa
 		DestructibleMesh->AssetImportData = ConstructObject<UAssetImportData>(UAssetImportData::StaticClass(), DestructibleMesh);
 		DestructibleMesh->AssetImportData->SourceFilePath = FReimportManager::SanitizeImportFilename(UFactory::CurrentFilename, DestructibleMesh);
 		DestructibleMesh->AssetImportData->SourceFileTimestamp = IFileManager::Get().GetTimeStamp(*UFactory::CurrentFilename).ToString();
+		DestructibleMesh->AssetImportData->bDirty = false;
 	}
 
 	DestructibleMesh->PreEditChange(NULL);
diff --git a/Engine/Source/Editor/UnrealEd/Private/Fbx/FbxSkeletalMeshImport.cpp b/Engine/Source/Editor/UnrealEd/Private/Fbx/FbxSkeletalMeshImport.cpp
--- a/Engine/Source/Editor/UnrealEd/Private/Fbx/FbxSkeletalMeshImport.cpp
+++ b/Engine/Source/Editor/UnrealEd/Private/Fbx/FbxSkeletalMeshImport.cpp
@@ -1278,6 +1278,7 @@ USkeletalMesh* UnFbx::FFbxImporter::ImportSkeletalMesh(UObject* InParent, TArray
 		UFbxSkeletalMeshImportData* ImportData = UFbxSkeletalMeshImportData::GetImportDataForSkeletalMesh(SkeletalMesh, TemplateImportData);
 		SkeletalMesh->AssetImportData->SourceFilePath = FReimportManager::SanitizeImportFilename(UFactory::CurrentFilename, SkeletalMesh);
 		SkeletalMesh->AssetImportData->SourceFileTimestamp = IFileManager::Get().GetTimeStamp(*UFactory::CurrentFilename).ToString();
+		SkeletalMesh->AssetImportData->bDirty = false;
 
 		SkeletalMesh->CalculateInvRefMatrices();
 		SkeletalMesh->PostEditChange();
diff --git a/Engine/Source/Editor/UnrealEd/Private/Fbx/FbxStaticMeshImport.cpp b/Engine/Source/Editor/UnrealEd/Private/Fbx/FbxStaticMeshImport.cpp
--- a/Engine/Source/Editor/UnrealEd/Private/Fbx/FbxStaticMeshImport.cpp
+++ b/Engine/Source/Editor/UnrealEd/Private/Fbx/FbxStaticMeshImport.cpp
@@ -1142,6 +1142,7 @@ UStaticMesh* UnFbx::FFbxImporter::ImportStaticMeshAsSingle(UObject* InParent, TA
 		UFbxStaticMeshImportData* ImportData = UFbxStaticMeshImportData::GetImportDataForStaticMesh(StaticMesh, TemplateImportData);
 		ImportData->SourceFilePath = FReimportManager::SanitizeImportFilename(UFactory::CurrentFilename, StaticMesh);
 		ImportData->SourceFileTimestamp = IFileManager::Get().GetTimeStamp(*UFactory::CurrentFilename).ToString();
+		ImportData->bDirty = false;
 
 		// @todo This overrides restored values currently but we need to be able to import over the existing settings if the user chose to do so.
 		SrcModel.BuildSettings.bRemoveDegenerates = ImportOptions->bRemoveDegenerates;
diff --git a/Engine/Source/Editor/UnrealEd/Private/SkeletalMeshEdit.cpp b/Engine/Source/Editor/UnrealEd/Private/SkeletalMeshEdit.cpp
--- a/Engine/Source/Editor/UnrealEd/Private/SkeletalMeshEdit.cpp
+++ b/Engine/Source/Editor/UnrealEd/Private/SkeletalMeshEdit.cpp
@@ -80,6 +80,7 @@ UAnimSequence * UEditorEngine::ImportFbxAnimation( USkeleton * Skeleton, UObject
 				UFbxAnimSequenceImportData* ImportData = UFbxAnimSequenceImportData::GetImportDataForAnimSequence(NewAnimation, TemplateImportData);
 				ImportData->SourceFilePath = FReimportManager::SanitizeImportFilename(UFactory::CurrentFilename, NewAnimation);
 				ImportData->SourceFileTimestamp = IFileManager::Get().GetTimeStamp(*UFactory::CurrentFilename).ToString();
+				ImportData->bDirty = false;
 			}
 		}
 	}
@@ -548,6 +549,7 @@ UAnimSequence * UnFbx::FFbxImporter::ImportAnimations(USkeleton * Skeleton, UObj
 		UFbxAnimSequenceImportData* ImportData = UFbxAnimSequenceImportData::GetImportDataForAnimSequence(DestSeq, TemplateImportData);
 		ImportData->SourceFilePath = FReimportManager::SanitizeImportFilename(UFactory::CurrentFilename, DestSeq);
 		ImportData->SourceFileTimestamp = IFileManager::Get().GetTimeStamp(*UFactory::CurrentFilename).ToString();
+		ImportData->bDirty = false;
 
 		ImportAnimation(Skeleton, DestSeq, Name, SortedLinks, NodeArray, CurAnimStack, ResampleRate, AnimTimeSpan);
 
diff --git a/Engine/Source/Runtime/Engine/Classes/EditorFramework/AssetImportData.h b/Engine/Source/Runtime/Engine/Classes/EditorFramework/AssetImportData.h
--- a/Engine/Source/Runtime/Engine/Classes/EditorFramework/AssetImportData.h
+++ b/Engine/Source/Runtime/Engine/Classes/EditorFramework/AssetImportData.h
@@ -18,6 +18,16 @@ class UAssetImportData : public UObject
 	/** Date/Time-stamp of the file from the last import */
 	UPROPERTY(VisibleAnywhere, Category=ImportSettings)
 	FString SourceFileTimestamp;
+
+	/** If true, the settings have been modified but not reimported yet */
+	UPROPERTY()
+	uint32 bDirty : 1;
+
+	// UObject interface
+#if WITH_EDITOR
+	ENGINE_API virtual void PostEditChangeProperty(struct FPropertyChangedEvent& PropertyChangedEvent) override;
+#endif // WITH_EDITOR
+	// End of UObject interface
 };
 
 
diff --git a/Engine/Source/Runtime/Engine/Private/EditorFramework/AssetImportData.cpp b/Engine/Source/Runtime/Engine/Private/EditorFramework/AssetImportData.cpp
--- a/Engine/Source/Runtime/Engine/Private/EditorFramework/AssetImportData.cpp
+++ b/Engine/Source/Runtime/Engine/Private/EditorFramework/AssetImportData.cpp
@@ -5,4 +5,14 @@
 UAssetImportData::UAssetImportData(const class FPostConstructInitializeProperties& PCIP)
 	: Super(PCIP)
 {
+	bDirty = false;
 }
+
+#if WITH_EDITOR
+void UAssetImportData::PostEditChangeProperty(FPropertyChangedEvent& PropertyChangedEvent)
+{
+	Super::PostEditChangeProperty(PropertyChangedEvent);
+
+	bDirty = true;
+}
+#endif // WITH_EDITOR
