Added analytics reporting to the static mesh editor

#ttp 340613 - FIXIF: EDITOR:ANALYTICS: Static Mesh Editor Instrumentation Pass

#branch UE4

#added Added analytics for various static mesh editor usage, making sure not to spam events when spinboxes are used and only report when the value actually changes.

#added Added analytics whenever an asset is reimported (if via an editor, the editor will then report that it was it that did the reimporting)

diff --git a/Engine/Source/Editor/DestructibleMeshEditor/Private/DestructibleMeshEditor.cpp b/Engine/Source/Editor/DestructibleMeshEditor/Private/DestructibleMeshEditor.cpp
--- a/Engine/Source/Editor/DestructibleMeshEditor/Private/DestructibleMeshEditor.cpp
+++ b/Engine/Source/Editor/DestructibleMeshEditor/Private/DestructibleMeshEditor.cpp
@@ -8,7 +8,8 @@
 #include "Editor/PropertyEditor/Public/PropertyEditorModule.h"
 #include "Editor/PropertyEditor/Public/IDetailsView.h"
 
-
+#include "Runtime/Analytics/Analytics/Public/Interfaces/IAnalyticsProvider.h"
+#include "EngineAnalytics.h"
 
 #define LOCTEXT_NAMESPACE "DestructibleMeshEditor"
 
@@ -397,6 +398,10 @@ void FDestructibleMeshEditor::OnPostReimport(UObject* InObject, bool bSuccess)
 
 	if ( bSuccess )
 	{
+		if (FEngineAnalytics::IsAvailable())
+		{
+			FEngineAnalytics::GetProvider().RecordEvent(TEXT("Editor.DestructibleMesh.ReimportedViaEditor"));
+		}
 		RefreshTool();
 	}
 }
diff --git a/Engine/Source/Editor/FontEditor/Private/FontEditor.cpp b/Engine/Source/Editor/FontEditor/Private/FontEditor.cpp
--- a/Engine/Source/Editor/FontEditor/Private/FontEditor.cpp
+++ b/Engine/Source/Editor/FontEditor/Private/FontEditor.cpp
@@ -12,6 +12,9 @@
 #include "Editor/PropertyEditor/Public/PropertyEditorModule.h"
 #include "Editor/PropertyEditor/Public/IDetailsView.h"
 
+#include "Runtime/Analytics/Analytics/Public/Interfaces/IAnalyticsProvider.h"
+#include "EngineAnalytics.h"
+
 #define LOCTEXT_NAMESPACE "FontEditor"
 
 DEFINE_LOG_CATEGORY_STATIC(LogFontEditor, Log, All);
@@ -718,6 +721,10 @@ void FFontEditor::OnPostReimport(UObject* InObject, bool bSuccess)
 
 	if ( bSuccess )
 	{
+		if (FEngineAnalytics::IsAvailable())
+		{
+			FEngineAnalytics::GetProvider().RecordEvent(TEXT("Editor.Font.ReimportedViaEditor"));
+		}
 		FontViewport->RefreshViewport();
 		FontPreviewWidget->RefreshViewport();
 	}
diff --git a/Engine/Source/Editor/Persona/Private/Persona.cpp b/Engine/Source/Editor/Persona/Private/Persona.cpp
--- a/Engine/Source/Editor/Persona/Private/Persona.cpp
+++ b/Engine/Source/Editor/Persona/Private/Persona.cpp
@@ -64,6 +64,9 @@
 #include "AnimationEditorUtils.h"
 #include "SAnimationSequenceBrowser.h"
 
+#include "Runtime/Analytics/Analytics/Public/Interfaces/IAnalyticsProvider.h"
+#include "EngineAnalytics.h"
+
 #define LOCTEXT_NAMESPACE "FPersona"
 
 /////////////////////////////////////////////////////
@@ -898,7 +901,14 @@ UDebugSkelMeshComponent* FPersona::GetPreviewMeshComponent()
 
 void FPersona::OnPostReimport(UObject* InObject, bool bSuccess)
 {
-	ConditionalRefreshEditor(InObject);
+	if (bSuccess)
+	{
+		if (FEngineAnalytics::IsAvailable())
+		{
+			FEngineAnalytics::GetProvider().RecordEvent(TEXT("Editor.Persona.ReimportedViaEditor"));
+		}
+		ConditionalRefreshEditor(InObject);
+	}
 }
 
 void FPersona::OnPostImport(UFactory* InFactory, UObject* InObject)
diff --git a/Engine/Source/Editor/StaticMeshEditor/Private/SStaticMeshEditorViewport.cpp b/Engine/Source/Editor/StaticMeshEditor/Private/SStaticMeshEditorViewport.cpp
--- a/Engine/Source/Editor/StaticMeshEditor/Private/SStaticMeshEditorViewport.cpp
+++ b/Engine/Source/Editor/StaticMeshEditor/Private/SStaticMeshEditorViewport.cpp
@@ -23,6 +23,9 @@
 #include "../Private/GeomFitUtils.h"
 #include "ComponentReregisterContext.h"
 
+#include "Runtime/Analytics/Analytics/Public/Interfaces/IAnalyticsProvider.h"
+#include "EngineAnalytics.h"
+
 #if WITH_PHYSX
 #include "Editor/UnrealEd/Private/EditorPhysXSupport.h"
 #endif
@@ -244,7 +247,10 @@ void SStaticMeshEditorViewport::SetViewModeWireframe()
 	{
 		CurrentViewMode = VMI_Lit;
 	}
-
+	if (FEngineAnalytics::IsAvailable())
+	{
+		FEngineAnalytics::GetProvider().RecordEvent(TEXT("Editor.StaticMesh.Toolbar"), TEXT("CurrentViewMode"), FString::Printf(TEXT("%d"), static_cast<int32>(CurrentViewMode)));
+	}
 	EditorViewportClient->SetViewMode(CurrentViewMode);
 	SceneViewport->Invalidate();
 
@@ -267,7 +273,10 @@ void SStaticMeshEditorViewport::SetViewModeVertexColor()
 		EditorViewportClient->EngineShowFlags.VertexColors = false;
 		EditorViewportClient->EngineShowFlags.Lighting = true;
 	}
-
+	if (FEngineAnalytics::IsAvailable())
+	{
+		FEngineAnalytics::GetProvider().RecordEvent(TEXT("Editor.StaticMesh.Toolbar"), FAnalyticsEventAttribute(TEXT("VertexColors"), EditorViewportClient->EngineShowFlags.VertexColors));
+	}
 	SceneViewport->Invalidate();
 }
 
diff --git a/Engine/Source/Editor/StaticMeshEditor/Private/StaticMeshEditor.cpp b/Engine/Source/Editor/StaticMeshEditor/Private/StaticMeshEditor.cpp
--- a/Engine/Source/Editor/StaticMeshEditor/Private/StaticMeshEditor.cpp
+++ b/Engine/Source/Editor/StaticMeshEditor/Private/StaticMeshEditor.cpp
@@ -22,6 +22,9 @@
 #include "Editor/ContentBrowser/Public/ContentBrowserModule.h"
 #include "Editor/WorkspaceMenuStructure/Public/WorkspaceMenuStructureModule.h"
 
+#include "Runtime/Analytics/Analytics/Public/Interfaces/IAnalyticsProvider.h"
+#include "EngineAnalytics.h"
+
 #define LOCTEXT_NAMESPACE "StaticMeshEditor"
 
 DEFINE_LOG_CATEGORY_STATIC(LogStaticMeshEditor, Log, All);
@@ -1062,6 +1065,10 @@ void FStaticMeshEditor::GenerateKDop(const FVector* Directions, uint32 NumDirect
 	GEditor->EndTransaction();
 	if (PrimIndex != INDEX_NONE)
 	{
+		if (FEngineAnalytics::IsAvailable())
+		{
+			FEngineAnalytics::GetProvider().RecordEvent(TEXT("Editor.StaticMesh.Collision"), TEXT("Type"), TEXT("KDop Collision"));
+		}
 		ClearSelectedPrims();
 		AddSelectedPrim(FPrimData(KPT_Convex, PrimIndex));
 	}
@@ -1076,6 +1083,10 @@ void FStaticMeshEditor::OnCollisionBox()
 	GEditor->EndTransaction();
 	if (PrimIndex != INDEX_NONE)
 	{
+		if (FEngineAnalytics::IsAvailable())
+		{
+			FEngineAnalytics::GetProvider().RecordEvent(TEXT("Editor.StaticMesh.Collision"), TEXT("Type"), TEXT("Box Collision"));
+		}
 		ClearSelectedPrims();
 		AddSelectedPrim(FPrimData(KPT_Box, PrimIndex));
 	}
@@ -1090,6 +1101,10 @@ void FStaticMeshEditor::OnCollisionSphere()
 	GEditor->EndTransaction();
 	if (PrimIndex != INDEX_NONE)
 	{
+		if (FEngineAnalytics::IsAvailable())
+		{
+			FEngineAnalytics::GetProvider().RecordEvent(TEXT("Editor.StaticMesh.Collision"), TEXT("Type"), TEXT("Sphere Collision"));
+		}
 		ClearSelectedPrims();
 		AddSelectedPrim(FPrimData(KPT_Sphere, PrimIndex));
 	}
@@ -1104,6 +1119,10 @@ void FStaticMeshEditor::OnCollisionSphyl()
 	GEditor->EndTransaction();
 	if (PrimIndex != INDEX_NONE)
 	{
+		if (FEngineAnalytics::IsAvailable())
+		{
+			FEngineAnalytics::GetProvider().RecordEvent(TEXT("Editor.StaticMesh.Collision"), TEXT("Type"), TEXT("Capsule Collision"));
+		}
 		ClearSelectedPrims();
 		AddSelectedPrim(FPrimData(KPT_Sphyl, PrimIndex));
 	}
@@ -1756,7 +1775,14 @@ void FStaticMeshEditor::OnPostReimport(UObject* InObject, bool bSuccess)
 		return;
 	}
 
-	RefreshTool();
+	if (bSuccess)
+	{
+		if (FEngineAnalytics::IsAvailable())
+		{
+			FEngineAnalytics::GetProvider().RecordEvent(TEXT("Editor.StaticMesh.ReimportedViaEditor"));
+		}
+		RefreshTool();
+	}
 }
 
 #undef LOCTEXT_NAMESPACE
diff --git a/Engine/Source/Editor/StaticMeshEditor/Private/StaticMeshEditorTools.cpp b/Engine/Source/Editor/StaticMeshEditor/Private/StaticMeshEditorTools.cpp
--- a/Engine/Source/Editor/StaticMeshEditor/Private/StaticMeshEditorTools.cpp
+++ b/Engine/Source/Editor/StaticMeshEditor/Private/StaticMeshEditorTools.cpp
@@ -12,6 +12,9 @@
 #include "FbxMeshUtils.h"
 #include "SVectorInputBox.h"
 
+#include "Runtime/Analytics/Analytics/Public/Interfaces/IAnalyticsProvider.h"
+#include "EngineAnalytics.h"
+
 const int32 DefaultHullCount = 4;
 const int32 DefaultVertsPerHull = 12;
 const int32 MaxHullCount = 24;
@@ -416,6 +419,7 @@ void FMeshBuildSettingsLayout::GenerateChildContent( IDetailChildrenBuilder& Chi
 			.MaxValue(100.0f)
 			.Value(this, &FMeshBuildSettingsLayout::GetDistanceFieldResolutionScale)
 			.OnValueChanged(this, &FMeshBuildSettingsLayout::OnDistanceFieldResolutionScaleChanged)
+			.OnValueCommitted(this, &FMeshBuildSettingsLayout::OnDistanceFieldResolutionScaleCommitted)
 		];
 	}
 		
@@ -492,44 +496,88 @@ float FMeshBuildSettingsLayout::GetDistanceFieldResolutionScale() const
 
 void FMeshBuildSettingsLayout::OnRecomputeNormalsChanged(ESlateCheckBoxState::Type NewState)
 {
-	BuildSettings.bRecomputeNormals = (NewState == ESlateCheckBoxState::Checked) ? true : false;
+	const bool bRecomputeNormals = (NewState == ESlateCheckBoxState::Checked) ? true : false;
+	if (BuildSettings.bRecomputeNormals != bRecomputeNormals)
+	{
+		if (FEngineAnalytics::IsAvailable())
+		{
+			FEngineAnalytics::GetProvider().RecordEvent(TEXT("Editor.StaticMesh.BuildSettings"), TEXT("bRecomputeNormals"), bRecomputeNormals ? TEXT("True") : TEXT("False"));
+		}
+		BuildSettings.bRecomputeNormals = bRecomputeNormals;
+	}
 }
 
 void FMeshBuildSettingsLayout::OnRecomputeTangentsChanged(ESlateCheckBoxState::Type NewState)
 {
-	BuildSettings.bRecomputeTangents = (NewState == ESlateCheckBoxState::Checked) ? true : false;
+	const bool bRecomputeTangents = (NewState == ESlateCheckBoxState::Checked) ? true : false;
+	if (BuildSettings.bRecomputeTangents != bRecomputeTangents)
+	{
+		if (FEngineAnalytics::IsAvailable())
+		{
+			FEngineAnalytics::GetProvider().RecordEvent(TEXT("Editor.StaticMesh.BuildSettings"), TEXT("bRecomputeTangents"), bRecomputeTangents ? TEXT("True") : TEXT("False"));
+		}
+		BuildSettings.bRecomputeTangents = bRecomputeTangents;
+	}
 }
 
 void FMeshBuildSettingsLayout::OnRemoveDegeneratesChanged(ESlateCheckBoxState::Type NewState)
 {
-	BuildSettings.bRemoveDegenerates = (NewState == ESlateCheckBoxState::Checked) ? true : false;
+	const bool bRemoveDegenerates = (NewState == ESlateCheckBoxState::Checked) ? true : false;
+	if (BuildSettings.bRemoveDegenerates != bRemoveDegenerates)
+	{
+		if (FEngineAnalytics::IsAvailable())
+		{
+			FEngineAnalytics::GetProvider().RecordEvent(TEXT("Editor.StaticMesh.BuildSettings"), TEXT("bRemoveDegenerates"), bRemoveDegenerates ? TEXT("True") : TEXT("False"));
+		}
+		BuildSettings.bRemoveDegenerates = bRemoveDegenerates;
+	}
 }
 
 void FMeshBuildSettingsLayout::OnUseFullPrecisionUVsChanged(ESlateCheckBoxState::Type NewState)
 {
-	BuildSettings.bUseFullPrecisionUVs = (NewState == ESlateCheckBoxState::Checked) ? true : false;
+	const bool bUseFullPrecisionUVs = (NewState == ESlateCheckBoxState::Checked) ? true : false;
+	if (BuildSettings.bUseFullPrecisionUVs != bUseFullPrecisionUVs)
+	{
+		if (FEngineAnalytics::IsAvailable())
+		{
+			FEngineAnalytics::GetProvider().RecordEvent(TEXT("Editor.StaticMesh.BuildSettings"), TEXT("bUseFullPrecisionUVs"), bUseFullPrecisionUVs ? TEXT("True") : TEXT("False"));
+		}
+		BuildSettings.bUseFullPrecisionUVs = bUseFullPrecisionUVs;
+	}
 }
 
 void FMeshBuildSettingsLayout::OnBuildScaleXChanged( float NewScaleX, ETextCommit::Type TextCommitType )
 {
-	if (!FMath::IsNearlyEqual(NewScaleX, 0.0f))
+	if (!FMath::IsNearlyEqual(NewScaleX, 0.0f) && BuildSettings.BuildScale3D.X != NewScaleX)
 	{
+		if (FEngineAnalytics::IsAvailable())
+		{
+			FEngineAnalytics::GetProvider().RecordEvent(TEXT("Editor.StaticMesh.BuildSettings"), TEXT("BuildScale3D.X"), FString::Printf(TEXT("%.3f"), NewScaleX));
+		}
 		BuildSettings.BuildScale3D.X = NewScaleX;
 	}
 }
 
 void FMeshBuildSettingsLayout::OnBuildScaleYChanged( float NewScaleY, ETextCommit::Type TextCommitType )
 {
-	if (!FMath::IsNearlyEqual(NewScaleY, 0.0f))
+	if (!FMath::IsNearlyEqual(NewScaleY, 0.0f) && BuildSettings.BuildScale3D.Y != NewScaleY)
 	{
+		if (FEngineAnalytics::IsAvailable())
+		{
+			FEngineAnalytics::GetProvider().RecordEvent(TEXT("Editor.StaticMesh.BuildSettings"), TEXT("BuildScale3D.Y"), FString::Printf(TEXT("%.3f"), NewScaleY));
+		}
 		BuildSettings.BuildScale3D.Y = NewScaleY;
 	}
 }
 
 void FMeshBuildSettingsLayout::OnBuildScaleZChanged( float NewScaleZ, ETextCommit::Type TextCommitType )
 {
-	if (!FMath::IsNearlyEqual(NewScaleZ, 0.0f))
+	if (!FMath::IsNearlyEqual(NewScaleZ, 0.0f) && BuildSettings.BuildScale3D.Z != NewScaleZ)
 	{
+		if (FEngineAnalytics::IsAvailable())
+		{
+			FEngineAnalytics::GetProvider().RecordEvent(TEXT("Editor.StaticMesh.BuildSettings"), TEXT("BuildScale3D.Z"), FString::Printf(TEXT("%.3f"), NewScaleZ));
+		}
 		BuildSettings.BuildScale3D.Z = NewScaleZ;
 	}
 }
@@ -539,6 +587,15 @@ void FMeshBuildSettingsLayout::OnDistanceFieldResolutionScaleChanged(float NewVa
 	BuildSettings.DistanceFieldResolutionScale = NewValue;
 }
 
+void FMeshBuildSettingsLayout::OnDistanceFieldResolutionScaleCommitted(float NewValue, ETextCommit::Type TextCommitType)
+{
+	if (FEngineAnalytics::IsAvailable())
+	{
+		FEngineAnalytics::GetProvider().RecordEvent(TEXT("Editor.StaticMesh.BuildSettings"), TEXT("DistanceFieldResolutionScale"), FString::Printf(TEXT("%.3f"), NewValue));
+	}
+	OnDistanceFieldResolutionScaleChanged(NewValue);
+}
+
 FMeshReductionSettingsLayout::FMeshReductionSettingsLayout( TSharedRef<FLevelOfDetailSettingsLayout> InParentLODSettings )
 	: ParentLODSettings( InParentLODSettings )
 {
@@ -579,6 +636,7 @@ void FMeshReductionSettingsLayout::GenerateChildContent( IDetailChildrenBuilder&
 			.MaxValue(100.0f)
 			.Value(this, &FMeshReductionSettingsLayout::GetPercentTriangles)
 			.OnValueChanged(this, &FMeshReductionSettingsLayout::OnPercentTrianglesChanged)
+			.OnValueCommitted(this, &FMeshReductionSettingsLayout::OnPercentTrianglesCommitted)
 		];
 
 	}
@@ -599,6 +657,7 @@ void FMeshReductionSettingsLayout::GenerateChildContent( IDetailChildrenBuilder&
 			.MaxValue(1000.0f)
 			.Value(this, &FMeshReductionSettingsLayout::GetMaxDeviation)
 			.OnValueChanged(this, &FMeshReductionSettingsLayout::OnMaxDeviationChanged)
+			.OnValueCommitted(this, &FMeshReductionSettingsLayout::OnMaxDeviationCommitted)
 		];
 
 	}
@@ -678,6 +737,8 @@ void FMeshReductionSettingsLayout::GenerateChildContent( IDetailChildrenBuilder&
 			.MinValue(0.0f)
 			.MaxValue(10.0f)
 			.Value(this, &FMeshReductionSettingsLayout::GetWeldingThreshold)
+			.OnValueChanged(this, &FMeshReductionSettingsLayout::OnWeldingThresholdChanged)
+			.OnValueCommitted(this, &FMeshReductionSettingsLayout::OnWeldingThresholdCommitted)
 		];
 
 	}
@@ -715,6 +776,7 @@ void FMeshReductionSettingsLayout::GenerateChildContent( IDetailChildrenBuilder&
 			.MaxValue(180.0f)
 			.Value(this, &FMeshReductionSettingsLayout::GetHardAngleThreshold)
 			.OnValueChanged(this, &FMeshReductionSettingsLayout::OnHardAngleThresholdChanged)
+			.OnValueCommitted(this, &FMeshReductionSettingsLayout::OnHardAngleThresholdCommitted)
 		];
 
 	}
@@ -791,19 +853,54 @@ void FMeshReductionSettingsLayout::OnPercentTrianglesChanged(float NewValue)
 	ReductionSettings.PercentTriangles = NewValue * 0.01f;
 }
 
+void FMeshReductionSettingsLayout::OnPercentTrianglesCommitted(float NewValue, ETextCommit::Type TextCommitType)
+{
+	if (FEngineAnalytics::IsAvailable())
+	{
+		FEngineAnalytics::GetProvider().RecordEvent(TEXT("Editor.StaticMesh.ReductionSettings"), TEXT("PercentTriangles"), FString::Printf(TEXT("%.1f"), NewValue));
+	}
+	OnPercentTrianglesChanged(NewValue);
+}
+
 void FMeshReductionSettingsLayout::OnMaxDeviationChanged(float NewValue)
 {
 	ReductionSettings.MaxDeviation = NewValue;
 }
 
+void FMeshReductionSettingsLayout::OnMaxDeviationCommitted(float NewValue, ETextCommit::Type TextCommitType)
+{
+	if (FEngineAnalytics::IsAvailable())
+	{
+		FEngineAnalytics::GetProvider().RecordEvent(TEXT("Editor.StaticMesh.ReductionSettings"), TEXT("MaxDeviation"), FString::Printf(TEXT("%.1f"), NewValue));
+	}
+	OnMaxDeviationChanged(NewValue);
+}
+
 void FMeshReductionSettingsLayout::OnWeldingThresholdChanged(float NewValue)
 {
 	ReductionSettings.WeldingThreshold = NewValue;
 }
 
+void FMeshReductionSettingsLayout::OnWeldingThresholdCommitted(float NewValue, ETextCommit::Type TextCommitType)
+{
+	if (FEngineAnalytics::IsAvailable())
+	{
+		FEngineAnalytics::GetProvider().RecordEvent(TEXT("Editor.StaticMesh.ReductionSettings"), TEXT("WeldingThreshold"), FString::Printf(TEXT("%.2f"), NewValue));
+	}
+	OnWeldingThresholdChanged(NewValue);
+}
+
 void FMeshReductionSettingsLayout::OnRecalculateNormalsChanged(ESlateCheckBoxState::Type NewValue)
 {
-	ReductionSettings.bRecalculateNormals = NewValue == ESlateCheckBoxState::Checked;
+	const bool bRecalculateNormals = NewValue == ESlateCheckBoxState::Checked;
+	if (ReductionSettings.bRecalculateNormals != bRecalculateNormals)
+	{
+		if (FEngineAnalytics::IsAvailable())
+		{
+			FEngineAnalytics::GetProvider().RecordEvent(TEXT("Editor.StaticMesh.ReductionSettings"), TEXT("bRecalculateNormals"), bRecalculateNormals ? TEXT("True") : TEXT("False"));
+		}
+		ReductionSettings.bRecalculateNormals = bRecalculateNormals;
+	}
 }
 
 void FMeshReductionSettingsLayout::OnHardAngleThresholdChanged(float NewValue)
@@ -811,19 +908,52 @@ void FMeshReductionSettingsLayout::OnHardAngleThresholdChanged(float NewValue)
 	ReductionSettings.HardAngleThreshold = NewValue;
 }
 
+void FMeshReductionSettingsLayout::OnHardAngleThresholdCommitted(float NewValue, ETextCommit::Type TextCommitType)
+{
+	if (FEngineAnalytics::IsAvailable())
+	{
+		FEngineAnalytics::GetProvider().RecordEvent(TEXT("Editor.StaticMesh.ReductionSettings"), TEXT("HardAngleThreshold"), FString::Printf(TEXT("%.3f"), NewValue));
+	}
+	OnHardAngleThresholdChanged(NewValue);
+}
+
 void FMeshReductionSettingsLayout::OnSilhouetteImportanceChanged(TSharedPtr<FString> NewValue, ESelectInfo::Type SelectInfo)
 {
-	ReductionSettings.SilhouetteImportance = (EMeshFeatureImportance::Type)ImportanceOptions.Find(NewValue);
+	const EMeshFeatureImportance::Type SilhouetteImportance = (EMeshFeatureImportance::Type)ImportanceOptions.Find(NewValue);
+	if (ReductionSettings.SilhouetteImportance != SilhouetteImportance)
+	{
+		if (FEngineAnalytics::IsAvailable())
+		{
+			FEngineAnalytics::GetProvider().RecordEvent(TEXT("Editor.StaticMesh.ReductionSettings"), TEXT("SilhouetteImportance"), *NewValue.Get());
+		}
+		ReductionSettings.SilhouetteImportance = SilhouetteImportance;
+	}
 }
 
 void FMeshReductionSettingsLayout::OnTextureImportanceChanged(TSharedPtr<FString> NewValue, ESelectInfo::Type SelectInfo)
 {
-	ReductionSettings.TextureImportance = (EMeshFeatureImportance::Type)ImportanceOptions.Find(NewValue);
+	const EMeshFeatureImportance::Type TextureImportance = (EMeshFeatureImportance::Type)ImportanceOptions.Find(NewValue);
+	if (ReductionSettings.TextureImportance != TextureImportance)
+	{
+		if (FEngineAnalytics::IsAvailable())
+		{
+			FEngineAnalytics::GetProvider().RecordEvent(TEXT("Editor.StaticMesh.ReductionSettings"), TEXT("TextureImportance"), *NewValue.Get());
+		}
+		ReductionSettings.TextureImportance = TextureImportance;
+	}
 }
 
 void FMeshReductionSettingsLayout::OnShadingImportanceChanged(TSharedPtr<FString> NewValue, ESelectInfo::Type SelectInfo)
 {
-	ReductionSettings.ShadingImportance = (EMeshFeatureImportance::Type)ImportanceOptions.Find(NewValue);
+	const EMeshFeatureImportance::Type ShadingImportance = (EMeshFeatureImportance::Type)ImportanceOptions.Find(NewValue);
+	if (ReductionSettings.ShadingImportance != ShadingImportance)
+	{
+		if (FEngineAnalytics::IsAvailable())
+		{
+			FEngineAnalytics::GetProvider().RecordEvent(TEXT("Editor.StaticMesh.ReductionSettings"), TEXT("ShadingImportance"), *NewValue.Get());
+		}
+		ReductionSettings.ShadingImportance = ShadingImportance;
+	}
 }
 
 FMeshSectionSettingsLayout::~FMeshSectionSettingsLayout()
diff --git a/Engine/Source/Editor/StaticMeshEditor/Private/StaticMeshEditorTools.h b/Engine/Source/Editor/StaticMeshEditor/Private/StaticMeshEditorTools.h
--- a/Engine/Source/Editor/StaticMeshEditor/Private/StaticMeshEditorTools.h
+++ b/Engine/Source/Editor/StaticMeshEditor/Private/StaticMeshEditorTools.h
@@ -157,6 +157,7 @@ private:
 	void OnBuildScaleZChanged( float NewScaleZ, ETextCommit::Type TextCommitType );
 
 	void OnDistanceFieldResolutionScaleChanged(float NewValue);
+	void OnDistanceFieldResolutionScaleCommitted(float NewValue, ETextCommit::Type TextCommitType);
 private:
 	TWeakPtr<FLevelOfDetailSettingsLayout> ParentLODSettings;
 	FMeshBuildSettings BuildSettings;
@@ -188,12 +189,15 @@ private:
 	float GetHardAngleThreshold() const;
 
 	void OnPercentTrianglesChanged(float NewValue);
+	void OnPercentTrianglesCommitted(float NewValue, ETextCommit::Type TextCommitType);
 	void OnMaxDeviationChanged(float NewValue);
+	void OnMaxDeviationCommitted(float NewValue, ETextCommit::Type TextCommitType);
 	void OnReductionAmountChanged(float NewValue);
 	void OnRecalculateNormalsChanged(ESlateCheckBoxState::Type NewValue);
 	void OnWeldingThresholdChanged(float NewValue);
-	
+	void OnWeldingThresholdCommitted(float NewValue, ETextCommit::Type TextCommitType);
 	void OnHardAngleThresholdChanged(float NewValue);
+	void OnHardAngleThresholdCommitted(float NewValue, ETextCommit::Type TextCommitType);
 
 	void OnSilhouetteImportanceChanged(TSharedPtr<FString> NewValue, ESelectInfo::Type SelectInfo);
 	void OnTextureImportanceChanged(TSharedPtr<FString> NewValue, ESelectInfo::Type SelectInfo);
diff --git a/Engine/Source/Editor/StaticMeshEditor/Private/StaticMeshEditorViewportClient.cpp b/Engine/Source/Editor/StaticMeshEditor/Private/StaticMeshEditorViewportClient.cpp
--- a/Engine/Source/Editor/StaticMeshEditor/Private/StaticMeshEditorViewportClient.cpp
+++ b/Engine/Source/Editor/StaticMeshEditor/Private/StaticMeshEditorViewportClient.cpp
@@ -17,6 +17,9 @@
 #include "ObjectTools.h"
 #include "SStaticMeshEditorViewport.h"
 
+#include "Runtime/Analytics/Analytics/Public/Interfaces/IAnalyticsProvider.h"
+#include "EngineAnalytics.h"
+
 #if WITH_PHYSX
 #include "Editor/UnrealEd/Private/EditorPhysXSupport.h"
 #endif
@@ -1142,6 +1145,10 @@ void FStaticMeshEditorViewportClient::SetPreviewMesh(UStaticMesh* InStaticMesh,
 void FStaticMeshEditorViewportClient::SetDrawUVOverlay()
 {
 	bDrawUVs = !bDrawUVs;
+	if (FEngineAnalytics::IsAvailable())
+	{
+		FEngineAnalytics::GetProvider().RecordEvent(TEXT("Editor.StaticMesh.Toolbar"), TEXT("bDrawUVs"), bDrawUVs ? TEXT("True") : TEXT("False"));
+	}
 	Invalidate();
 }
 
@@ -1153,6 +1160,10 @@ bool FStaticMeshEditorViewportClient::IsSetDrawUVOverlayChecked() const
 void FStaticMeshEditorViewportClient::SetShowNormals()
 {
 	bDrawNormals = !bDrawNormals;
+	if (FEngineAnalytics::IsAvailable())
+	{
+		FEngineAnalytics::GetProvider().RecordEvent(TEXT("Editor.StaticMesh.Toolbar"), TEXT("bDrawNormals"), bDrawNormals ? TEXT("True") : TEXT("False"));
+	}
 	Invalidate();
 }
 
@@ -1164,6 +1175,10 @@ bool FStaticMeshEditorViewportClient::IsSetShowNormalsChecked() const
 void FStaticMeshEditorViewportClient::SetShowTangents()
 {
 	bDrawTangents = !bDrawTangents;
+	if (FEngineAnalytics::IsAvailable())
+	{
+		FEngineAnalytics::GetProvider().RecordEvent(TEXT("Editor.StaticMesh.Toolbar"), TEXT("bDrawTangents"), bDrawTangents ? TEXT("True") : TEXT("False"));
+	}
 	Invalidate();
 }
 
@@ -1175,6 +1190,10 @@ bool FStaticMeshEditorViewportClient::IsSetShowTangentsChecked() const
 void FStaticMeshEditorViewportClient::SetShowBinormals()
 {
 	bDrawBinormals = !bDrawBinormals;
+	if (FEngineAnalytics::IsAvailable())
+	{
+		FEngineAnalytics::GetProvider().RecordEvent(TEXT("Editor.StaticMesh.Toolbar"), TEXT("bDrawBinormals"), bDrawBinormals ? TEXT("True") : TEXT("False"));
+	}
 	Invalidate();
 }
 
@@ -1186,6 +1205,10 @@ bool FStaticMeshEditorViewportClient::IsSetShowBinormalsChecked() const
 void FStaticMeshEditorViewportClient::SetShowWireframeCollision()
 {
 	bShowCollision = !bShowCollision;
+	if (FEngineAnalytics::IsAvailable())
+	{
+		FEngineAnalytics::GetProvider().RecordEvent(TEXT("Editor.StaticMesh.Toolbar"), TEXT("bShowCollision"), bShowCollision ? TEXT("True") : TEXT("False"));
+	}
 	StaticMeshEditorPtr.Pin()->ClearSelectedPrims();
 	Invalidate();
 }
@@ -1198,6 +1221,10 @@ bool FStaticMeshEditorViewportClient::IsSetShowWireframeCollisionChecked() const
 void FStaticMeshEditorViewportClient::SetShowSockets()
 {
 	bShowSockets = !bShowSockets;
+	if (FEngineAnalytics::IsAvailable())
+	{
+		FEngineAnalytics::GetProvider().RecordEvent(TEXT("Editor.StaticMesh.Toolbar"), TEXT("bShowSockets"), bShowSockets ? TEXT("True") : TEXT("False"));
+	}
 	Invalidate();
 }
 bool FStaticMeshEditorViewportClient::IsSetShowSocketsChecked() const
@@ -1208,6 +1235,10 @@ bool FStaticMeshEditorViewportClient::IsSetShowSocketsChecked() const
 void FStaticMeshEditorViewportClient::SetShowPivot()
 {
 	bShowPivot = !bShowPivot;
+	if (FEngineAnalytics::IsAvailable())
+	{
+		FEngineAnalytics::GetProvider().RecordEvent(TEXT("Editor.StaticMesh.Toolbar"), TEXT("bShowPivot"), bShowPivot ? TEXT("True") : TEXT("False"));
+	}
 	Invalidate();
 }
 
@@ -1219,6 +1250,10 @@ bool FStaticMeshEditorViewportClient::IsSetShowPivotChecked() const
 void FStaticMeshEditorViewportClient::SetDrawAdditionalData()
 {
 	bDrawAdditionalData = !bDrawAdditionalData;
+	if (FEngineAnalytics::IsAvailable())
+	{
+		FEngineAnalytics::GetProvider().RecordEvent(TEXT("Editor.StaticMesh.Toolbar"), TEXT("bDrawAdditionalData"), bDrawAdditionalData ? TEXT("True") : TEXT("False"));
+	}
 	Invalidate();
 }
 
diff --git a/Engine/Source/Editor/TextureEditor/Private/TextureEditorToolkit.cpp b/Engine/Source/Editor/TextureEditor/Private/TextureEditorToolkit.cpp
--- a/Engine/Source/Editor/TextureEditor/Private/TextureEditorToolkit.cpp
+++ b/Engine/Source/Editor/TextureEditor/Private/TextureEditorToolkit.cpp
@@ -8,6 +8,8 @@
 
 #include "Factories.h"
 
+#include "Runtime/Analytics/Analytics/Public/Interfaces/IAnalyticsProvider.h"
+#include "EngineAnalytics.h"
 
 #define LOCTEXT_NAMESPACE "FTextureEditorToolkit"
 
@@ -842,7 +844,14 @@ void FTextureEditorToolkit::OnPostReimport(UObject* InObject, bool bSuccess)
 		return;
 	}
 
-	if (!bSuccess)
+	if (bSuccess)
+	{
+		if (FEngineAnalytics::IsAvailable())
+		{
+			FEngineAnalytics::GetProvider().RecordEvent(TEXT("Editor.Texture.ReimportedViaEditor"));
+		}
+	}
+	else
 	{
 		// Failed, restore the compression flag
 		Texture->DeferCompression = SavedCompressionSetting;
diff --git a/Engine/Source/Editor/UnrealEd/Private/Editor.cpp b/Engine/Source/Editor/UnrealEd/Private/Editor.cpp
--- a/Engine/Source/Editor/UnrealEd/Private/Editor.cpp
+++ b/Engine/Source/Editor/UnrealEd/Private/Editor.cpp
@@ -81,6 +81,9 @@
 #include "EditorStyle.h"
 #include "EngineBuildSettings.h"
 
+#include "Runtime/Analytics/Analytics/Public/Interfaces/IAnalyticsProvider.h"
+#include "EngineAnalytics.h"
+
 // AIMdule
 #include "BehaviorTree/BehaviorTreeManager.h"
 
@@ -2588,6 +2591,10 @@ bool FReimportManager::Reimport( UObject* Obj, bool bAskForNewFileIfMissing )
 					{
 						Obj->PostEditChange();
 						GEditor->BroadcastObjectReimported(Obj);
+						if (FEngineAnalytics::IsAvailable())
+						{
+							FEngineAnalytics::GetProvider().RecordEvent(TEXT("Editor.AssetReimported"));
+						}
 						bSuccess = true;
 					}
 					else if( Result == EReimportResult::Cancelled )
diff --git a/Engine/Source/Editor/UnrealEd/Private/EditorViewportClient.cpp b/Engine/Source/Editor/UnrealEd/Private/EditorViewportClient.cpp
--- a/Engine/Source/Editor/UnrealEd/Private/EditorViewportClient.cpp
+++ b/Engine/Source/Editor/UnrealEd/Private/EditorViewportClient.cpp
@@ -3562,7 +3562,11 @@ bool FEditorViewportClient::IsCameraLocked() const
 
 void FEditorViewportClient::SetShowGrid()
 {
-	DrawHelper.bDrawGrid = !DrawHelper.bDrawGrid;	
+	DrawHelper.bDrawGrid = !DrawHelper.bDrawGrid;
+	if (FEngineAnalytics::IsAvailable())
+	{
+		FEngineAnalytics::GetProvider().RecordEvent(TEXT("Editor.StaticMesh.Toolbar"), TEXT("bDrawGrid"), DrawHelper.bDrawGrid ? TEXT("True") : TEXT("False"));
+	}
 	Invalidate();
 }
 
@@ -3579,6 +3583,10 @@ void FEditorViewportClient::SetShowBounds(bool bShow)
 void FEditorViewportClient::ToggleShowBounds()
 {
 	EngineShowFlags.Bounds = 1 - EngineShowFlags.Bounds;
+	if (FEngineAnalytics::IsAvailable())
+	{
+		FEngineAnalytics::GetProvider().RecordEvent(TEXT("Editor.StaticMesh.Toolbar"), TEXT("Bounds"), FString::Printf(TEXT("%d"), EngineShowFlags.Bounds));
+	}
 	Invalidate();
 }
 
diff --git a/Engine/Source/Editor/UnrealEd/Private/SSocketManager.cpp b/Engine/Source/Editor/UnrealEd/Private/SSocketManager.cpp
--- a/Engine/Source/Editor/UnrealEd/Private/SSocketManager.cpp
+++ b/Engine/Source/Editor/UnrealEd/Private/SSocketManager.cpp
@@ -10,6 +10,9 @@
 #include "SSocketManager.h"
 #include "ScopedTransaction.h"
 
+#include "Runtime/Analytics/Analytics/Public/Interfaces/IAnalyticsProvider.h"
+#include "EngineAnalytics.h"
+
 #define LOCTEXT_NAMESPACE "SSCSSocketManagerEditor"
 
 struct SocketListItem
@@ -416,6 +419,11 @@ void SSocketManager::CreateSocket()
 		UStaticMeshSocket* NewSocket = ConstructObject<UStaticMeshSocket>(UStaticMeshSocket::StaticClass(), CurrentStaticMesh);
 		check(NewSocket);
 
+		if (FEngineAnalytics::IsAvailable())
+		{
+			FEngineAnalytics::GetProvider().RecordEvent(TEXT("Editor.StaticMesh.CreateSocket"));
+		}
+
 		FString SocketNameString = TEXT("Socket");
 		FName SocketName = FName(*SocketNameString);
 
