Modified analytics to include .Usage in the name

diff --git a/Engine/Source/Editor/Cascade/Private/Cascade.cpp b/Engine/Source/Editor/Cascade/Private/Cascade.cpp
--- a/Engine/Source/Editor/Cascade/Private/Cascade.cpp
+++ b/Engine/Source/Editor/Cascade/Private/Cascade.cpp
@@ -213,7 +213,7 @@ void FCascade::InitCascade(const EToolkitMode::Type Mode, const TSharedPtr< clas
 	}
 	if (Description.Len() > 0 && FEngineAnalytics::IsAvailable())
 	{
-		FEngineAnalytics::GetProvider().RecordEvent(TEXT("Editor.Cascade.Init"), TEXT("Overview"), Description);
+		FEngineAnalytics::GetProvider().RecordEvent(TEXT("Editor.Usage.Cascade.Init"), TEXT("Overview"), Description);
 	}
 
 	ParticleSystemComponent = ConstructObject<UCascadeParticleSystemComponent>(UCascadeParticleSystemComponent::StaticClass());
@@ -944,7 +944,7 @@ void FCascade::OnNewModule(int32 Idx)
 
 	if (FEngineAnalytics::IsAvailable())
 	{
-		FEngineAnalytics::GetProvider().RecordEvent(TEXT("Editor.Cascade.NewModule"), TEXT("Class"), NewModClass->GetName());
+		FEngineAnalytics::GetProvider().RecordEvent(TEXT("Editor.Usage.Cascade.NewModule"), TEXT("Class"), NewModClass->GetName());
 	}
 
 	ParticleSystem->MarkPackageDirty();
@@ -1111,7 +1111,7 @@ void FCascade::OnNewEmitter()
 
 	if (FEngineAnalytics::IsAvailable())
 	{
-		FEngineAnalytics::GetProvider().RecordEvent(TEXT("Editor.Cascade.NewEmitter"));
+		FEngineAnalytics::GetProvider().RecordEvent(TEXT("Editor.Usage.Cascade.NewEmitter"));
 	}
 
 	// Refresh viewport
@@ -3337,7 +3337,7 @@ void FCascade::AddLOD(bool bBeforeCurrent)
 
 		if (FEngineAnalytics::IsAvailable())
 		{
-			FEngineAnalytics::GetProvider().RecordEvent(TEXT("Editor.Cascade.NewLOD"), FAnalyticsEventAttribute(TEXT("Index"), CurrentLODIndex));
+			FEngineAnalytics::GetProvider().RecordEvent(TEXT("Editor.Usage.Cascade.NewLOD"), FAnalyticsEventAttribute(TEXT("Index"), CurrentLODIndex));
 		}
 
 		UpdateLODLevel();
@@ -4230,7 +4230,7 @@ void FCascade::OnDeleteLOD()
 
 	if (FEngineAnalytics::IsAvailable())
 	{
-		FEngineAnalytics::GetProvider().RecordEvent(TEXT("Editor.Cascade.DeleteLOD"), FAnalyticsEventAttribute(TEXT("Index"), Selection));
+		FEngineAnalytics::GetProvider().RecordEvent(TEXT("Editor.Usage.Cascade.DeleteLOD"), FAnalyticsEventAttribute(TEXT("Index"), Selection));
 	}
 
 	ForceUpdate();
@@ -4421,7 +4421,7 @@ void FCascade::OnDeleteModule(bool bConfirm)
 
 	if (FEngineAnalytics::IsAvailable())
 	{
-		FEngineAnalytics::GetProvider().RecordEvent(TEXT("Editor.Cascade.DeleteModule"), TEXT("Class"), ModuleName);
+		FEngineAnalytics::GetProvider().RecordEvent(TEXT("Editor.Usage.Cascade.DeleteModule"), TEXT("Class"), ModuleName);
 	}
 
 	SetSelectedEmitter(SelectedEmitter);
@@ -4938,7 +4938,7 @@ void FCascade::OnDeleteEmitter()
 
 	if (FEngineAnalytics::IsAvailable())
 	{
-		FEngineAnalytics::GetProvider().RecordEvent(TEXT("Editor.Cascade.DeleteEmitter"));
+		FEngineAnalytics::GetProvider().RecordEvent(TEXT("Editor.Usage.Cascade.DeleteEmitter"));
 	}
 
 	ParticleSystem->MarkPackageDirty();
diff --git a/Engine/Source/Editor/DestructibleMeshEditor/Private/DestructibleMeshEditor.cpp b/Engine/Source/Editor/DestructibleMeshEditor/Private/DestructibleMeshEditor.cpp
--- a/Engine/Source/Editor/DestructibleMeshEditor/Private/DestructibleMeshEditor.cpp
+++ b/Engine/Source/Editor/DestructibleMeshEditor/Private/DestructibleMeshEditor.cpp
@@ -400,7 +400,7 @@ void FDestructibleMeshEditor::OnPostReimport(UObject* InObject, bool bSuccess)
 	{
 		if (FEngineAnalytics::IsAvailable())
 		{
-			FEngineAnalytics::GetProvider().RecordEvent(TEXT("Editor.DestructibleMesh.ReimportedViaEditor"));
+			FEngineAnalytics::GetProvider().RecordEvent(TEXT("Editor.Usage.DestructibleMesh.ReimportedViaEditor"));
 		}
 		RefreshTool();
 	}
diff --git a/Engine/Source/Editor/FontEditor/Private/FontEditor.cpp b/Engine/Source/Editor/FontEditor/Private/FontEditor.cpp
--- a/Engine/Source/Editor/FontEditor/Private/FontEditor.cpp
+++ b/Engine/Source/Editor/FontEditor/Private/FontEditor.cpp
@@ -723,7 +723,7 @@ void FFontEditor::OnPostReimport(UObject* InObject, bool bSuccess)
 	{
 		if (FEngineAnalytics::IsAvailable())
 		{
-			FEngineAnalytics::GetProvider().RecordEvent(TEXT("Editor.Font.ReimportedViaEditor"));
+			FEngineAnalytics::GetProvider().RecordEvent(TEXT("Editor.Usage.Font.ReimportedViaEditor"));
 		}
 		FontViewport->RefreshViewport();
 		FontPreviewWidget->RefreshViewport();
diff --git a/Engine/Source/Editor/Matinee/Private/Matinee.cpp b/Engine/Source/Editor/Matinee/Private/Matinee.cpp
--- a/Engine/Source/Editor/Matinee/Private/Matinee.cpp
+++ b/Engine/Source/Editor/Matinee/Private/Matinee.cpp
@@ -1560,7 +1560,7 @@ void FMatinee::ToggleRecordInterpValues(void)
 
 		if (FEngineAnalytics::IsAvailable())
 		{
-			FEngineAnalytics::GetProvider().RecordEvent(TEXT("Editor.Matinee.Recorded"));
+			FEngineAnalytics::GetProvider().RecordEvent(TEXT("Editor.Usage.Matinee.Recorded"));
 		}
 	}
 	else
diff --git a/Engine/Source/Editor/Matinee/Private/MatineeMenus.cpp b/Engine/Source/Editor/Matinee/Private/MatineeMenus.cpp
--- a/Engine/Source/Editor/Matinee/Private/MatineeMenus.cpp
+++ b/Engine/Source/Editor/Matinee/Private/MatineeMenus.cpp
@@ -407,7 +407,7 @@ void FMatinee::NewGroupPopupTextCommitted(
 
 	if (FEngineAnalytics::IsAvailable())
 	{
-		FEngineAnalytics::GetProvider().RecordEvent(TEXT("Editor.Matinee.NewGroup"), GroupAttribs);
+		FEngineAnalytics::GetProvider().RecordEvent(TEXT("Editor.Usage.Matinee.NewGroup"), GroupAttribs);
 	}
 
 	IData->InterpGroups.Append(NewGroups);
diff --git a/Engine/Source/Editor/Matinee/Private/MatineeTools.cpp b/Engine/Source/Editor/Matinee/Private/MatineeTools.cpp
--- a/Engine/Source/Editor/Matinee/Private/MatineeTools.cpp
+++ b/Engine/Source/Editor/Matinee/Private/MatineeTools.cpp
@@ -2780,7 +2780,7 @@ UInterpTrack* FMatinee::AddTrackToGroup( UInterpGroup* Group, UClass* TrackClass
 
 	if (FEngineAnalytics::IsAvailable())
 	{
-		FEngineAnalytics::GetProvider().RecordEvent(TEXT("Editor.Matinee.NewTrack"), TEXT("Class"), TrackClass->GetName());
+		FEngineAnalytics::GetProvider().RecordEvent(TEXT("Editor.Usage.Matinee.NewTrack"), TEXT("Class"), TrackClass->GetName());
 	}
 
 	OutNewTrackIndex = Group->InterpTracks.Add(NewTrack);
@@ -2954,7 +2954,7 @@ void FMatinee::DeleteSelectedTracks()
 		{
 			if (FEngineAnalytics::IsAvailable())
 			{
-				FEngineAnalytics::GetProvider().RecordEvent(TEXT("Editor.Matinee.DelTrack"), TEXT("Class"), ActiveTrack->GetClass()->GetName());
+				FEngineAnalytics::GetProvider().RecordEvent(TEXT("Editor.Usage.Matinee.DelTrack"), TEXT("Class"), ActiveTrack->GetClass()->GetName());
 			}
 
 			ActiveTrack->Modify();
@@ -3060,7 +3060,7 @@ void FMatinee::DeleteSelectedGroups()
 		UInterpGroup* GroupToDelete = *GroupIt;
 		if (FEngineAnalytics::IsAvailable())
 		{
-			FEngineAnalytics::GetProvider().RecordEvent(TEXT("Editor.Matinee.DelGroup"), TEXT("Name"), GroupToDelete->GroupName.ToString());
+			FEngineAnalytics::GetProvider().RecordEvent(TEXT("Editor.Usage.Matinee.DelGroup"), TEXT("Name"), GroupToDelete->GroupName.ToString());
 		}
 
 		// Mark InterpGroup and all InterpTracks as Modified.
@@ -3070,7 +3070,7 @@ void FMatinee::DeleteSelectedGroups()
 			UInterpTrack* Track = GroupToDelete->InterpTracks[j];
 			if (FEngineAnalytics::IsAvailable())
 			{
-				FEngineAnalytics::GetProvider().RecordEvent(TEXT("Editor.Matinee.DelTrack"), TEXT("Class"), Track->GetClass()->GetName());
+				FEngineAnalytics::GetProvider().RecordEvent(TEXT("Editor.Usage.Matinee.DelTrack"), TEXT("Class"), Track->GetClass()->GetName());
 			}
 
 			Track->Modify();
diff --git a/Engine/Source/Editor/Persona/Private/Persona.cpp b/Engine/Source/Editor/Persona/Private/Persona.cpp
--- a/Engine/Source/Editor/Persona/Private/Persona.cpp
+++ b/Engine/Source/Editor/Persona/Private/Persona.cpp
@@ -905,7 +905,7 @@ void FPersona::OnPostReimport(UObject* InObject, bool bSuccess)
 	{
 		if (FEngineAnalytics::IsAvailable())
 		{
-			FEngineAnalytics::GetProvider().RecordEvent(TEXT("Editor.Persona.ReimportedViaEditor"));
+			FEngineAnalytics::GetProvider().RecordEvent(TEXT("Editor.Usage.Persona.ReimportedViaEditor"));
 		}
 		ConditionalRefreshEditor(InObject);
 	}
diff --git a/Engine/Source/Editor/StaticMeshEditor/Private/SStaticMeshEditorViewport.cpp b/Engine/Source/Editor/StaticMeshEditor/Private/SStaticMeshEditorViewport.cpp
--- a/Engine/Source/Editor/StaticMeshEditor/Private/SStaticMeshEditorViewport.cpp
+++ b/Engine/Source/Editor/StaticMeshEditor/Private/SStaticMeshEditorViewport.cpp
@@ -249,7 +249,7 @@ void SStaticMeshEditorViewport::SetViewModeWireframe()
 	}
 	if (FEngineAnalytics::IsAvailable())
 	{
-		FEngineAnalytics::GetProvider().RecordEvent(TEXT("Editor.StaticMesh.Toolbar"), TEXT("CurrentViewMode"), FString::Printf(TEXT("%d"), static_cast<int32>(CurrentViewMode)));
+		FEngineAnalytics::GetProvider().RecordEvent(TEXT("Editor.Usage.StaticMesh.Toolbar"), TEXT("CurrentViewMode"), FString::Printf(TEXT("%d"), static_cast<int32>(CurrentViewMode)));
 	}
 	EditorViewportClient->SetViewMode(CurrentViewMode);
 	SceneViewport->Invalidate();
@@ -275,7 +275,7 @@ void SStaticMeshEditorViewport::SetViewModeVertexColor()
 	}
 	if (FEngineAnalytics::IsAvailable())
 	{
-		FEngineAnalytics::GetProvider().RecordEvent(TEXT("Editor.StaticMesh.Toolbar"), FAnalyticsEventAttribute(TEXT("VertexColors"), EditorViewportClient->EngineShowFlags.VertexColors));
+		FEngineAnalytics::GetProvider().RecordEvent(TEXT("Editor.Usage.StaticMesh.Toolbar"), FAnalyticsEventAttribute(TEXT("VertexColors"), EditorViewportClient->EngineShowFlags.VertexColors));
 	}
 	SceneViewport->Invalidate();
 }
diff --git a/Engine/Source/Editor/StaticMeshEditor/Private/StaticMeshEditor.cpp b/Engine/Source/Editor/StaticMeshEditor/Private/StaticMeshEditor.cpp
--- a/Engine/Source/Editor/StaticMeshEditor/Private/StaticMeshEditor.cpp
+++ b/Engine/Source/Editor/StaticMeshEditor/Private/StaticMeshEditor.cpp
@@ -1067,7 +1067,7 @@ void FStaticMeshEditor::GenerateKDop(const FVector* Directions, uint32 NumDirect
 	{
 		if (FEngineAnalytics::IsAvailable())
 		{
-			FEngineAnalytics::GetProvider().RecordEvent(TEXT("Editor.StaticMesh.Collision"), TEXT("Type"), TEXT("KDop Collision"));
+			FEngineAnalytics::GetProvider().RecordEvent(TEXT("Editor.Usage.StaticMesh.Collision"), TEXT("Type"), TEXT("KDop Collision"));
 		}
 		ClearSelectedPrims();
 		AddSelectedPrim(FPrimData(KPT_Convex, PrimIndex));
@@ -1085,7 +1085,7 @@ void FStaticMeshEditor::OnCollisionBox()
 	{
 		if (FEngineAnalytics::IsAvailable())
 		{
-			FEngineAnalytics::GetProvider().RecordEvent(TEXT("Editor.StaticMesh.Collision"), TEXT("Type"), TEXT("Box Collision"));
+			FEngineAnalytics::GetProvider().RecordEvent(TEXT("Editor.Usage.StaticMesh.Collision"), TEXT("Type"), TEXT("Box Collision"));
 		}
 		ClearSelectedPrims();
 		AddSelectedPrim(FPrimData(KPT_Box, PrimIndex));
@@ -1103,7 +1103,7 @@ void FStaticMeshEditor::OnCollisionSphere()
 	{
 		if (FEngineAnalytics::IsAvailable())
 		{
-			FEngineAnalytics::GetProvider().RecordEvent(TEXT("Editor.StaticMesh.Collision"), TEXT("Type"), TEXT("Sphere Collision"));
+			FEngineAnalytics::GetProvider().RecordEvent(TEXT("Editor.Usage.StaticMesh.Collision"), TEXT("Type"), TEXT("Sphere Collision"));
 		}
 		ClearSelectedPrims();
 		AddSelectedPrim(FPrimData(KPT_Sphere, PrimIndex));
@@ -1121,7 +1121,7 @@ void FStaticMeshEditor::OnCollisionSphyl()
 	{
 		if (FEngineAnalytics::IsAvailable())
 		{
-			FEngineAnalytics::GetProvider().RecordEvent(TEXT("Editor.StaticMesh.Collision"), TEXT("Type"), TEXT("Capsule Collision"));
+			FEngineAnalytics::GetProvider().RecordEvent(TEXT("Editor.Usage.StaticMesh.Collision"), TEXT("Type"), TEXT("Capsule Collision"));
 		}
 		ClearSelectedPrims();
 		AddSelectedPrim(FPrimData(KPT_Sphyl, PrimIndex));
@@ -1779,7 +1779,7 @@ void FStaticMeshEditor::OnPostReimport(UObject* InObject, bool bSuccess)
 	{
 		if (FEngineAnalytics::IsAvailable())
 		{
-			FEngineAnalytics::GetProvider().RecordEvent(TEXT("Editor.StaticMesh.ReimportedViaEditor"));
+			FEngineAnalytics::GetProvider().RecordEvent(TEXT("Editor.Usage.StaticMesh.ReimportedViaEditor"));
 		}
 		RefreshTool();
 	}
diff --git a/Engine/Source/Editor/StaticMeshEditor/Private/StaticMeshEditorTools.cpp b/Engine/Source/Editor/StaticMeshEditor/Private/StaticMeshEditorTools.cpp
--- a/Engine/Source/Editor/StaticMeshEditor/Private/StaticMeshEditorTools.cpp
+++ b/Engine/Source/Editor/StaticMeshEditor/Private/StaticMeshEditorTools.cpp
@@ -501,7 +501,7 @@ void FMeshBuildSettingsLayout::OnRecomputeNormalsChanged(ESlateCheckBoxState::Ty
 	{
 		if (FEngineAnalytics::IsAvailable())
 		{
-			FEngineAnalytics::GetProvider().RecordEvent(TEXT("Editor.StaticMesh.BuildSettings"), TEXT("bRecomputeNormals"), bRecomputeNormals ? TEXT("True") : TEXT("False"));
+			FEngineAnalytics::GetProvider().RecordEvent(TEXT("Editor.Usage.StaticMesh.BuildSettings"), TEXT("bRecomputeNormals"), bRecomputeNormals ? TEXT("True") : TEXT("False"));
 		}
 		BuildSettings.bRecomputeNormals = bRecomputeNormals;
 	}
@@ -514,7 +514,7 @@ void FMeshBuildSettingsLayout::OnRecomputeTangentsChanged(ESlateCheckBoxState::T
 	{
 		if (FEngineAnalytics::IsAvailable())
 		{
-			FEngineAnalytics::GetProvider().RecordEvent(TEXT("Editor.StaticMesh.BuildSettings"), TEXT("bRecomputeTangents"), bRecomputeTangents ? TEXT("True") : TEXT("False"));
+			FEngineAnalytics::GetProvider().RecordEvent(TEXT("Editor.Usage.StaticMesh.BuildSettings"), TEXT("bRecomputeTangents"), bRecomputeTangents ? TEXT("True") : TEXT("False"));
 		}
 		BuildSettings.bRecomputeTangents = bRecomputeTangents;
 	}
@@ -527,7 +527,7 @@ void FMeshBuildSettingsLayout::OnRemoveDegeneratesChanged(ESlateCheckBoxState::T
 	{
 		if (FEngineAnalytics::IsAvailable())
 		{
-			FEngineAnalytics::GetProvider().RecordEvent(TEXT("Editor.StaticMesh.BuildSettings"), TEXT("bRemoveDegenerates"), bRemoveDegenerates ? TEXT("True") : TEXT("False"));
+			FEngineAnalytics::GetProvider().RecordEvent(TEXT("Editor.Usage.StaticMesh.BuildSettings"), TEXT("bRemoveDegenerates"), bRemoveDegenerates ? TEXT("True") : TEXT("False"));
 		}
 		BuildSettings.bRemoveDegenerates = bRemoveDegenerates;
 	}
@@ -540,7 +540,7 @@ void FMeshBuildSettingsLayout::OnUseFullPrecisionUVsChanged(ESlateCheckBoxState:
 	{
 		if (FEngineAnalytics::IsAvailable())
 		{
-			FEngineAnalytics::GetProvider().RecordEvent(TEXT("Editor.StaticMesh.BuildSettings"), TEXT("bUseFullPrecisionUVs"), bUseFullPrecisionUVs ? TEXT("True") : TEXT("False"));
+			FEngineAnalytics::GetProvider().RecordEvent(TEXT("Editor.Usage.StaticMesh.BuildSettings"), TEXT("bUseFullPrecisionUVs"), bUseFullPrecisionUVs ? TEXT("True") : TEXT("False"));
 		}
 		BuildSettings.bUseFullPrecisionUVs = bUseFullPrecisionUVs;
 	}
@@ -552,7 +552,7 @@ void FMeshBuildSettingsLayout::OnBuildScaleXChanged( float NewScaleX, ETextCommi
 	{
 		if (FEngineAnalytics::IsAvailable())
 		{
-			FEngineAnalytics::GetProvider().RecordEvent(TEXT("Editor.StaticMesh.BuildSettings"), TEXT("BuildScale3D.X"), FString::Printf(TEXT("%.3f"), NewScaleX));
+			FEngineAnalytics::GetProvider().RecordEvent(TEXT("Editor.Usage.StaticMesh.BuildSettings"), TEXT("BuildScale3D.X"), FString::Printf(TEXT("%.3f"), NewScaleX));
 		}
 		BuildSettings.BuildScale3D.X = NewScaleX;
 	}
@@ -564,7 +564,7 @@ void FMeshBuildSettingsLayout::OnBuildScaleYChanged( float NewScaleY, ETextCommi
 	{
 		if (FEngineAnalytics::IsAvailable())
 		{
-			FEngineAnalytics::GetProvider().RecordEvent(TEXT("Editor.StaticMesh.BuildSettings"), TEXT("BuildScale3D.Y"), FString::Printf(TEXT("%.3f"), NewScaleY));
+			FEngineAnalytics::GetProvider().RecordEvent(TEXT("Editor.Usage.StaticMesh.BuildSettings"), TEXT("BuildScale3D.Y"), FString::Printf(TEXT("%.3f"), NewScaleY));
 		}
 		BuildSettings.BuildScale3D.Y = NewScaleY;
 	}
@@ -576,7 +576,7 @@ void FMeshBuildSettingsLayout::OnBuildScaleZChanged( float NewScaleZ, ETextCommi
 	{
 		if (FEngineAnalytics::IsAvailable())
 		{
-			FEngineAnalytics::GetProvider().RecordEvent(TEXT("Editor.StaticMesh.BuildSettings"), TEXT("BuildScale3D.Z"), FString::Printf(TEXT("%.3f"), NewScaleZ));
+			FEngineAnalytics::GetProvider().RecordEvent(TEXT("Editor.Usage.StaticMesh.BuildSettings"), TEXT("BuildScale3D.Z"), FString::Printf(TEXT("%.3f"), NewScaleZ));
 		}
 		BuildSettings.BuildScale3D.Z = NewScaleZ;
 	}
@@ -591,7 +591,7 @@ void FMeshBuildSettingsLayout::OnDistanceFieldResolutionScaleCommitted(float New
 {
 	if (FEngineAnalytics::IsAvailable())
 	{
-		FEngineAnalytics::GetProvider().RecordEvent(TEXT("Editor.StaticMesh.BuildSettings"), TEXT("DistanceFieldResolutionScale"), FString::Printf(TEXT("%.3f"), NewValue));
+		FEngineAnalytics::GetProvider().RecordEvent(TEXT("Editor.Usage.StaticMesh.BuildSettings"), TEXT("DistanceFieldResolutionScale"), FString::Printf(TEXT("%.3f"), NewValue));
 	}
 	OnDistanceFieldResolutionScaleChanged(NewValue);
 }
diff --git a/Engine/Source/Editor/StaticMeshEditor/Private/StaticMeshEditorViewportClient.cpp b/Engine/Source/Editor/StaticMeshEditor/Private/StaticMeshEditorViewportClient.cpp
--- a/Engine/Source/Editor/StaticMeshEditor/Private/StaticMeshEditorViewportClient.cpp
+++ b/Engine/Source/Editor/StaticMeshEditor/Private/StaticMeshEditorViewportClient.cpp
@@ -1147,7 +1147,7 @@ void FStaticMeshEditorViewportClient::SetDrawUVOverlay()
 	bDrawUVs = !bDrawUVs;
 	if (FEngineAnalytics::IsAvailable())
 	{
-		FEngineAnalytics::GetProvider().RecordEvent(TEXT("Editor.StaticMesh.Toolbar"), TEXT("bDrawUVs"), bDrawUVs ? TEXT("True") : TEXT("False"));
+		FEngineAnalytics::GetProvider().RecordEvent(TEXT("Editor.Usage.StaticMesh.Toolbar"), TEXT("bDrawUVs"), bDrawUVs ? TEXT("True") : TEXT("False"));
 	}
 	Invalidate();
 }
@@ -1162,7 +1162,7 @@ void FStaticMeshEditorViewportClient::SetShowNormals()
 	bDrawNormals = !bDrawNormals;
 	if (FEngineAnalytics::IsAvailable())
 	{
-		FEngineAnalytics::GetProvider().RecordEvent(TEXT("Editor.StaticMesh.Toolbar"), TEXT("bDrawNormals"), bDrawNormals ? TEXT("True") : TEXT("False"));
+		FEngineAnalytics::GetProvider().RecordEvent(TEXT("Editor.Usage.StaticMesh.Toolbar"), TEXT("bDrawNormals"), bDrawNormals ? TEXT("True") : TEXT("False"));
 	}
 	Invalidate();
 }
@@ -1177,7 +1177,7 @@ void FStaticMeshEditorViewportClient::SetShowTangents()
 	bDrawTangents = !bDrawTangents;
 	if (FEngineAnalytics::IsAvailable())
 	{
-		FEngineAnalytics::GetProvider().RecordEvent(TEXT("Editor.StaticMesh.Toolbar"), TEXT("bDrawTangents"), bDrawTangents ? TEXT("True") : TEXT("False"));
+		FEngineAnalytics::GetProvider().RecordEvent(TEXT("Editor.Usage.StaticMesh.Toolbar"), TEXT("bDrawTangents"), bDrawTangents ? TEXT("True") : TEXT("False"));
 	}
 	Invalidate();
 }
@@ -1192,7 +1192,7 @@ void FStaticMeshEditorViewportClient::SetShowBinormals()
 	bDrawBinormals = !bDrawBinormals;
 	if (FEngineAnalytics::IsAvailable())
 	{
-		FEngineAnalytics::GetProvider().RecordEvent(TEXT("Editor.StaticMesh.Toolbar"), TEXT("bDrawBinormals"), bDrawBinormals ? TEXT("True") : TEXT("False"));
+		FEngineAnalytics::GetProvider().RecordEvent(TEXT("Editor.Usage.StaticMesh.Toolbar"), TEXT("bDrawBinormals"), bDrawBinormals ? TEXT("True") : TEXT("False"));
 	}
 	Invalidate();
 }
@@ -1207,7 +1207,7 @@ void FStaticMeshEditorViewportClient::SetShowWireframeCollision()
 	bShowCollision = !bShowCollision;
 	if (FEngineAnalytics::IsAvailable())
 	{
-		FEngineAnalytics::GetProvider().RecordEvent(TEXT("Editor.StaticMesh.Toolbar"), TEXT("bShowCollision"), bShowCollision ? TEXT("True") : TEXT("False"));
+		FEngineAnalytics::GetProvider().RecordEvent(TEXT("Editor.Usage.StaticMesh.Toolbar"), TEXT("bShowCollision"), bShowCollision ? TEXT("True") : TEXT("False"));
 	}
 	StaticMeshEditorPtr.Pin()->ClearSelectedPrims();
 	Invalidate();
@@ -1223,7 +1223,7 @@ void FStaticMeshEditorViewportClient::SetShowSockets()
 	bShowSockets = !bShowSockets;
 	if (FEngineAnalytics::IsAvailable())
 	{
-		FEngineAnalytics::GetProvider().RecordEvent(TEXT("Editor.StaticMesh.Toolbar"), TEXT("bShowSockets"), bShowSockets ? TEXT("True") : TEXT("False"));
+		FEngineAnalytics::GetProvider().RecordEvent(TEXT("Editor.Usage.StaticMesh.Toolbar"), TEXT("bShowSockets"), bShowSockets ? TEXT("True") : TEXT("False"));
 	}
 	Invalidate();
 }
@@ -1237,7 +1237,7 @@ void FStaticMeshEditorViewportClient::SetShowPivot()
 	bShowPivot = !bShowPivot;
 	if (FEngineAnalytics::IsAvailable())
 	{
-		FEngineAnalytics::GetProvider().RecordEvent(TEXT("Editor.StaticMesh.Toolbar"), TEXT("bShowPivot"), bShowPivot ? TEXT("True") : TEXT("False"));
+		FEngineAnalytics::GetProvider().RecordEvent(TEXT("Editor.Usage.StaticMesh.Toolbar"), TEXT("bShowPivot"), bShowPivot ? TEXT("True") : TEXT("False"));
 	}
 	Invalidate();
 }
@@ -1252,7 +1252,7 @@ void FStaticMeshEditorViewportClient::SetDrawAdditionalData()
 	bDrawAdditionalData = !bDrawAdditionalData;
 	if (FEngineAnalytics::IsAvailable())
 	{
-		FEngineAnalytics::GetProvider().RecordEvent(TEXT("Editor.StaticMesh.Toolbar"), TEXT("bDrawAdditionalData"), bDrawAdditionalData ? TEXT("True") : TEXT("False"));
+		FEngineAnalytics::GetProvider().RecordEvent(TEXT("Editor.Usage.StaticMesh.Toolbar"), TEXT("bDrawAdditionalData"), bDrawAdditionalData ? TEXT("True") : TEXT("False"));
 	}
 	Invalidate();
 }
diff --git a/Engine/Source/Editor/TextureEditor/Private/TextureEditorToolkit.cpp b/Engine/Source/Editor/TextureEditor/Private/TextureEditorToolkit.cpp
--- a/Engine/Source/Editor/TextureEditor/Private/TextureEditorToolkit.cpp
+++ b/Engine/Source/Editor/TextureEditor/Private/TextureEditorToolkit.cpp
@@ -998,7 +998,7 @@ void FTextureEditorToolkit::HandleReimportManagerPostReimport( UObject* InObject
 	{
 		if (FEngineAnalytics::IsAvailable())
 		{
-			FEngineAnalytics::GetProvider().RecordEvent(TEXT("Editor.Texture.ReimportedViaEditor"));
+			FEngineAnalytics::GetProvider().RecordEvent(TEXT("Editor.Usage.Texture.ReimportedViaEditor"));
 		}
 	}
 	else
diff --git a/Engine/Source/Editor/UnrealEd/Private/Editor.cpp b/Engine/Source/Editor/UnrealEd/Private/Editor.cpp
--- a/Engine/Source/Editor/UnrealEd/Private/Editor.cpp
+++ b/Engine/Source/Editor/UnrealEd/Private/Editor.cpp
@@ -2603,7 +2603,7 @@ bool FReimportManager::Reimport( UObject* Obj, bool bAskForNewFileIfMissing )
 						GEditor->BroadcastObjectReimported(Obj);
 						if (FEngineAnalytics::IsAvailable())
 						{
-							FEngineAnalytics::GetProvider().RecordEvent(TEXT("Editor.AssetReimported"));
+							FEngineAnalytics::GetProvider().RecordEvent(TEXT("Editor.Usage.AssetReimported"));
 						}
 						bSuccess = true;
 					}
diff --git a/Engine/Source/Editor/UnrealEd/Private/EditorViewportClient.cpp b/Engine/Source/Editor/UnrealEd/Private/EditorViewportClient.cpp
--- a/Engine/Source/Editor/UnrealEd/Private/EditorViewportClient.cpp
+++ b/Engine/Source/Editor/UnrealEd/Private/EditorViewportClient.cpp
@@ -3570,7 +3570,7 @@ void FEditorViewportClient::SetShowGrid()
 	DrawHelper.bDrawGrid = !DrawHelper.bDrawGrid;
 	if (FEngineAnalytics::IsAvailable())
 	{
-		FEngineAnalytics::GetProvider().RecordEvent(TEXT("Editor.StaticMesh.Toolbar"), TEXT("bDrawGrid"), DrawHelper.bDrawGrid ? TEXT("True") : TEXT("False"));
+		FEngineAnalytics::GetProvider().RecordEvent(TEXT("Editor.Usage.StaticMesh.Toolbar"), TEXT("bDrawGrid"), DrawHelper.bDrawGrid ? TEXT("True") : TEXT("False"));
 	}
 	Invalidate();
 }
@@ -3590,7 +3590,7 @@ void FEditorViewportClient::ToggleShowBounds()
 	EngineShowFlags.Bounds = 1 - EngineShowFlags.Bounds;
 	if (FEngineAnalytics::IsAvailable())
 	{
-		FEngineAnalytics::GetProvider().RecordEvent(TEXT("Editor.StaticMesh.Toolbar"), TEXT("Bounds"), FString::Printf(TEXT("%d"), EngineShowFlags.Bounds));
+		FEngineAnalytics::GetProvider().RecordEvent(TEXT("Editor.Usage.StaticMesh.Toolbar"), TEXT("Bounds"), FString::Printf(TEXT("%d"), EngineShowFlags.Bounds));
 	}
 	Invalidate();
 }
diff --git a/Engine/Source/Editor/UnrealEd/Private/SSocketManager.cpp b/Engine/Source/Editor/UnrealEd/Private/SSocketManager.cpp
--- a/Engine/Source/Editor/UnrealEd/Private/SSocketManager.cpp
+++ b/Engine/Source/Editor/UnrealEd/Private/SSocketManager.cpp
@@ -421,7 +421,7 @@ void SSocketManager::CreateSocket()
 
 		if (FEngineAnalytics::IsAvailable())
 		{
-			FEngineAnalytics::GetProvider().RecordEvent(TEXT("Editor.StaticMesh.CreateSocket"));
+			FEngineAnalytics::GetProvider().RecordEvent(TEXT("Editor.Usage.StaticMesh.CreateSocket"));
 		}
 
 		FString SocketNameString = TEXT("Socket");
