More modified analytics to include .Usage in the name

diff --git a/Engine/Source/Editor/Matinee/Private/MatineeMenus.cpp b/Engine/Source/Editor/Matinee/Private/MatineeMenus.cpp
--- a/Engine/Source/Editor/Matinee/Private/MatineeMenus.cpp
+++ b/Engine/Source/Editor/Matinee/Private/MatineeMenus.cpp
@@ -3766,7 +3766,7 @@ void FMatinee::OnMenuImport()
 					{
 						if (FEngineAnalytics::IsAvailable())
 						{
-							FEngineAnalytics::GetProvider().RecordEvent(TEXT("Editor.Matinee.Imported"));
+							FEngineAnalytics::GetProvider().RecordEvent(TEXT("Editor.Usage.Matinee.Imported"));
 						}
 					}
 
@@ -3842,7 +3842,7 @@ void FMatinee::OnMenuExport()
 				{
 					if (FEngineAnalytics::IsAvailable())
 					{
-						FEngineAnalytics::GetProvider().RecordEvent(TEXT("Editor.Matinee.Exported"));
+						FEngineAnalytics::GetProvider().RecordEvent(TEXT("Editor.Usage.Matinee.Exported"));
 					}
 				}
 
@@ -3989,7 +3989,7 @@ void FMatinee::OnExportSoundCueInfoCommand()
 
 			if (FEngineAnalytics::IsAvailable())
 			{
-				FEngineAnalytics::GetProvider().RecordEvent(TEXT("Editor.Matinee.ExportedSoundCue"));
+				FEngineAnalytics::GetProvider().RecordEvent(TEXT("Editor.Usage.Matinee.ExportedSoundCue"));
 			}
 		}
 		else
@@ -4178,7 +4178,7 @@ void FMatinee::OnExportAnimationInfoCommand()
 
 			if (FEngineAnalytics::IsAvailable())
 			{
-				FEngineAnalytics::GetProvider().RecordEvent(TEXT("Editor.Matinee.ExportedAnimationInfo"));
+				FEngineAnalytics::GetProvider().RecordEvent(TEXT("Editor.Usage.Matinee.ExportedAnimationInfo"));
 			}
 		}
 	}
diff --git a/Engine/Source/Editor/StaticMeshEditor/Private/StaticMeshEditorTools.cpp b/Engine/Source/Editor/StaticMeshEditor/Private/StaticMeshEditorTools.cpp
--- a/Engine/Source/Editor/StaticMeshEditor/Private/StaticMeshEditorTools.cpp
+++ b/Engine/Source/Editor/StaticMeshEditor/Private/StaticMeshEditorTools.cpp
@@ -857,7 +857,7 @@ void FMeshReductionSettingsLayout::OnPercentTrianglesCommitted(float NewValue, E
 {
 	if (FEngineAnalytics::IsAvailable())
 	{
-		FEngineAnalytics::GetProvider().RecordEvent(TEXT("Editor.StaticMesh.ReductionSettings"), TEXT("PercentTriangles"), FString::Printf(TEXT("%.1f"), NewValue));
+		FEngineAnalytics::GetProvider().RecordEvent(TEXT("Editor.Usage.StaticMesh.ReductionSettings"), TEXT("PercentTriangles"), FString::Printf(TEXT("%.1f"), NewValue));
 	}
 	OnPercentTrianglesChanged(NewValue);
 }
@@ -871,7 +871,7 @@ void FMeshReductionSettingsLayout::OnMaxDeviationCommitted(float NewValue, EText
 {
 	if (FEngineAnalytics::IsAvailable())
 	{
-		FEngineAnalytics::GetProvider().RecordEvent(TEXT("Editor.StaticMesh.ReductionSettings"), TEXT("MaxDeviation"), FString::Printf(TEXT("%.1f"), NewValue));
+		FEngineAnalytics::GetProvider().RecordEvent(TEXT("Editor.Usage.StaticMesh.ReductionSettings"), TEXT("MaxDeviation"), FString::Printf(TEXT("%.1f"), NewValue));
 	}
 	OnMaxDeviationChanged(NewValue);
 }
@@ -885,7 +885,7 @@ void FMeshReductionSettingsLayout::OnWeldingThresholdCommitted(float NewValue, E
 {
 	if (FEngineAnalytics::IsAvailable())
 	{
-		FEngineAnalytics::GetProvider().RecordEvent(TEXT("Editor.StaticMesh.ReductionSettings"), TEXT("WeldingThreshold"), FString::Printf(TEXT("%.2f"), NewValue));
+		FEngineAnalytics::GetProvider().RecordEvent(TEXT("Editor.Usage.StaticMesh.ReductionSettings"), TEXT("WeldingThreshold"), FString::Printf(TEXT("%.2f"), NewValue));
 	}
 	OnWeldingThresholdChanged(NewValue);
 }
@@ -897,7 +897,7 @@ void FMeshReductionSettingsLayout::OnRecalculateNormalsChanged(ESlateCheckBoxSta
 	{
 		if (FEngineAnalytics::IsAvailable())
 		{
-			FEngineAnalytics::GetProvider().RecordEvent(TEXT("Editor.StaticMesh.ReductionSettings"), TEXT("bRecalculateNormals"), bRecalculateNormals ? TEXT("True") : TEXT("False"));
+			FEngineAnalytics::GetProvider().RecordEvent(TEXT("Editor.Usage.StaticMesh.ReductionSettings"), TEXT("bRecalculateNormals"), bRecalculateNormals ? TEXT("True") : TEXT("False"));
 		}
 		ReductionSettings.bRecalculateNormals = bRecalculateNormals;
 	}
@@ -912,7 +912,7 @@ void FMeshReductionSettingsLayout::OnHardAngleThresholdCommitted(float NewValue,
 {
 	if (FEngineAnalytics::IsAvailable())
 	{
-		FEngineAnalytics::GetProvider().RecordEvent(TEXT("Editor.StaticMesh.ReductionSettings"), TEXT("HardAngleThreshold"), FString::Printf(TEXT("%.3f"), NewValue));
+		FEngineAnalytics::GetProvider().RecordEvent(TEXT("Editor.Usage.StaticMesh.ReductionSettings"), TEXT("HardAngleThreshold"), FString::Printf(TEXT("%.3f"), NewValue));
 	}
 	OnHardAngleThresholdChanged(NewValue);
 }
@@ -924,7 +924,7 @@ void FMeshReductionSettingsLayout::OnSilhouetteImportanceChanged(TSharedPtr<FStr
 	{
 		if (FEngineAnalytics::IsAvailable())
 		{
-			FEngineAnalytics::GetProvider().RecordEvent(TEXT("Editor.StaticMesh.ReductionSettings"), TEXT("SilhouetteImportance"), *NewValue.Get());
+			FEngineAnalytics::GetProvider().RecordEvent(TEXT("Editor.Usage.StaticMesh.ReductionSettings"), TEXT("SilhouetteImportance"), *NewValue.Get());
 		}
 		ReductionSettings.SilhouetteImportance = SilhouetteImportance;
 	}
@@ -937,7 +937,7 @@ void FMeshReductionSettingsLayout::OnTextureImportanceChanged(TSharedPtr<FString
 	{
 		if (FEngineAnalytics::IsAvailable())
 		{
-			FEngineAnalytics::GetProvider().RecordEvent(TEXT("Editor.StaticMesh.ReductionSettings"), TEXT("TextureImportance"), *NewValue.Get());
+			FEngineAnalytics::GetProvider().RecordEvent(TEXT("Editor.Usage.StaticMesh.ReductionSettings"), TEXT("TextureImportance"), *NewValue.Get());
 		}
 		ReductionSettings.TextureImportance = TextureImportance;
 	}
@@ -950,7 +950,7 @@ void FMeshReductionSettingsLayout::OnShadingImportanceChanged(TSharedPtr<FString
 	{
 		if (FEngineAnalytics::IsAvailable())
 		{
-			FEngineAnalytics::GetProvider().RecordEvent(TEXT("Editor.StaticMesh.ReductionSettings"), TEXT("ShadingImportance"), *NewValue.Get());
+			FEngineAnalytics::GetProvider().RecordEvent(TEXT("Editor.Usage.StaticMesh.ReductionSettings"), TEXT("ShadingImportance"), *NewValue.Get());
 		}
 		ReductionSettings.ShadingImportance = ShadingImportance;
 	}
