Fixed crash importing a static mesh with no triangles and then opening the mesh in the static mesh editor

diff --git a/Engine/Source/Editor/StaticMeshEditor/Private/StaticMeshEditorViewportClient.cpp b/Engine/Source/Editor/StaticMeshEditor/Private/StaticMeshEditorViewportClient.cpp
--- a/Engine/Source/Editor/StaticMeshEditor/Private/StaticMeshEditorViewportClient.cpp
+++ b/Engine/Source/Editor/StaticMeshEditor/Private/StaticMeshEditorViewportClient.cpp
@@ -709,25 +709,28 @@ void FStaticMeshEditorViewportClient::DrawCanvas( FViewport& InViewport, FSceneV
 	TextItems.Add(SStaticMeshEditorViewport::FOverlayTextItem(
 		FText::Format(NSLOCTEXT("UnrealEd", "UVChannels_F", "UV Channels:  {0}"), FText::AsNumber(StaticMeshEditorPtr.Pin()->GetNumUVChannels(CurrentLODLevel)))));
 
-	const FDistanceFieldVolumeData& VolumeData = StaticMesh->RenderData->LODResources[0].DistanceFieldData;
-
-	if (VolumeData.Size.GetMax() > 0)
+	if( StaticMesh->RenderData->LODResources.Num() > 0 )
 	{
-		float MemoryMb = (VolumeData.Size.X * VolumeData.Size.Y * VolumeData.Size.Z * VolumeData.DistanceFieldVolume.GetTypeSize()) / (1024.0f * 1024.0f);
-
-		FNumberFormattingOptions NumberOptions;
-		NumberOptions.MinimumFractionalDigits = 2;
-		NumberOptions.MaximumFractionalDigits = 2;
+		const FDistanceFieldVolumeData& VolumeData = StaticMesh->RenderData->LODResources[0].DistanceFieldData;
 
-		if (VolumeData.bMeshWasClosed)
-		{
-			TextItems.Add(SStaticMeshEditorViewport::FOverlayTextItem(
-				FText::Format(NSLOCTEXT("UnrealEd", "DistanceFieldRes_F", "Distance Field:  {0}x{1}x{2} = {3}Mb"), FText::AsNumber(VolumeData.Size.X), FText::AsNumber(VolumeData.Size.Y), FText::AsNumber(VolumeData.Size.Z), FText::AsNumber(MemoryMb, &NumberOptions))));
-		}
-		else
+		if(VolumeData.Size.GetMax() > 0)
 		{
-			TextItems.Add(SStaticMeshEditorViewport::FOverlayTextItem(
-				NSLOCTEXT("UnrealEd", "DistanceFieldClosed_F", "Distance Field:  Mesh was not closed and material was one-sided")));
+			float MemoryMb = (VolumeData.Size.X * VolumeData.Size.Y * VolumeData.Size.Z * VolumeData.DistanceFieldVolume.GetTypeSize()) / (1024.0f * 1024.0f);
+
+			FNumberFormattingOptions NumberOptions;
+			NumberOptions.MinimumFractionalDigits = 2;
+			NumberOptions.MaximumFractionalDigits = 2;
+
+			if(VolumeData.bMeshWasClosed)
+			{
+				TextItems.Add(SStaticMeshEditorViewport::FOverlayTextItem(
+					FText::Format(NSLOCTEXT("UnrealEd", "DistanceFieldRes_F", "Distance Field:  {0}x{1}x{2} = {3}Mb"), FText::AsNumber(VolumeData.Size.X), FText::AsNumber(VolumeData.Size.Y), FText::AsNumber(VolumeData.Size.Z), FText::AsNumber(MemoryMb, &NumberOptions))));
+			}
+			else
+			{
+				TextItems.Add(SStaticMeshEditorViewport::FOverlayTextItem(
+					NSLOCTEXT("UnrealEd", "DistanceFieldClosed_F", "Distance Field:  Mesh was not closed and material was one-sided")));
+			}
 		}
 	}
 
diff --git a/Engine/Source/Editor/UnrealEd/Private/Fbx/FbxStaticMeshImport.cpp b/Engine/Source/Editor/UnrealEd/Private/Fbx/FbxStaticMeshImport.cpp
--- a/Engine/Source/Editor/UnrealEd/Private/Fbx/FbxStaticMeshImport.cpp
+++ b/Engine/Source/Editor/UnrealEd/Private/Fbx/FbxStaticMeshImport.cpp
@@ -414,6 +414,13 @@ bool UnFbx::FFbxImporter::BuildStaticMeshFromGeometry(FbxMesh* Mesh, UStaticMesh
 	TotalMatrixForNormal = TotalMatrix.Inverse();
 	TotalMatrixForNormal = TotalMatrixForNormal.Transpose();
 	int32 TriangleCount = Mesh->GetPolygonCount();
+
+	if(TriangleCount == 0)
+	{
+		AddTokenizedErrorMessage(FTokenizedMessage::Create(EMessageSeverity::Error, FText::Format(LOCTEXT("Error_NoTrianglesFoundInMesh", "No triangles were found on mesh  '{0}'"), FText::FromString(Mesh->GetName()))), FFbxErrors::StaticMesh_NoTriangles);
+		return false;
+	}
+
 	int32 VertexCount = Mesh->GetControlPointsCount();
 	int32 WedgeCount = TriangleCount * 3;
 	bool OddNegativeScale = IsOddNegativeScale(TotalMatrix);
diff --git a/Engine/Source/Runtime/Core/Private/Misc/FbxErrors.cpp b/Engine/Source/Runtime/Core/Private/Misc/FbxErrors.cpp
--- a/Engine/Source/Runtime/Core/Private/Misc/FbxErrors.cpp
+++ b/Engine/Source/Runtime/Core/Private/Misc/FbxErrors.cpp
@@ -45,6 +45,8 @@ FName FFbxErrors::StaticMesh_TooManyMaterials(TEXT("StaticMesh_TooManyMaterials"
 
 FName FFbxErrors::StaticMesh_UVSetLayoutProblem(TEXT("StaticMesh_UVSetLayoutProblem"));
 
+FName FFbxErrors::StaticMesh_NoTriangles(TEXT("StaticMesh_NoTriangles"));
+
 /** SkeletalMesh **/
 FName FFbxErrors::SkeletalMesh_DifferentRoots(TEXT("SkeletalMesh_DifferentRoot"));
 
diff --git a/Engine/Source/Runtime/Core/Public/Misc/FbxErrors.h b/Engine/Source/Runtime/Core/Public/Misc/FbxErrors.h
--- a/Engine/Source/Runtime/Core/Public/Misc/FbxErrors.h
+++ b/Engine/Source/Runtime/Core/Public/Misc/FbxErrors.h
@@ -52,6 +52,8 @@ struct CORE_API FFbxErrors
 
 	static FName StaticMesh_UVSetLayoutProblem;
 
+	static FName StaticMesh_NoTriangles;
+
 	/** SkeletalMesh **/
 	static FName SkeletalMesh_DifferentRoots;
 
