Fixed crash in StaticMeshEditorViewportClient.cpp (Added a null check)

diff --git a/Engine/Source/Editor/StaticMeshEditor/Private/StaticMeshEditorViewportClient.cpp b/Engine/Source/Editor/StaticMeshEditor/Private/StaticMeshEditorViewportClient.cpp
--- a/Engine/Source/Editor/StaticMeshEditor/Private/StaticMeshEditorViewportClient.cpp
+++ b/Engine/Source/Editor/StaticMeshEditor/Private/StaticMeshEditorViewportClient.cpp
@@ -711,25 +711,28 @@ void FStaticMeshEditorViewportClient::DrawCanvas( FViewport& InViewport, FSceneV
 
 	if( StaticMesh->RenderData->LODResources.Num() > 0 )
 	{
-		const FDistanceFieldVolumeData& VolumeData = *(StaticMesh->RenderData->LODResources[0].DistanceFieldData);
-
-		if(VolumeData.Size.GetMax() > 0)
+		if (StaticMesh->RenderData->LODResources[0].DistanceFieldData != nullptr )
 		{
-			float MemoryMb = (VolumeData.Size.X * VolumeData.Size.Y * VolumeData.Size.Z * VolumeData.DistanceFieldVolume.GetTypeSize()) / (1024.0f * 1024.0f);
-
-			FNumberFormattingOptions NumberOptions;
-			NumberOptions.MinimumFractionalDigits = 2;
-			NumberOptions.MaximumFractionalDigits = 2;
+			const FDistanceFieldVolumeData& VolumeData = *(StaticMesh->RenderData->LODResources[0].DistanceFieldData);
 
-			if(VolumeData.bMeshWasClosed)
-			{
-				TextItems.Add(SStaticMeshEditorViewport::FOverlayTextItem(
-					FText::Format(NSLOCTEXT("UnrealEd", "DistanceFieldRes_F", "Distance Field:  {0}x{1}x{2} = {3}Mb"), FText::AsNumber(VolumeData.Size.X), FText::AsNumber(VolumeData.Size.Y), FText::AsNumber(VolumeData.Size.Z), FText::AsNumber(MemoryMb, &NumberOptions))));
-			}
-			else
+			if (VolumeData.Size.GetMax() > 0)
 			{
-				TextItems.Add(SStaticMeshEditorViewport::FOverlayTextItem(
-					NSLOCTEXT("UnrealEd", "DistanceFieldClosed_F", "Distance Field:  Mesh was not closed and material was one-sided")));
+				float MemoryMb = (VolumeData.Size.X * VolumeData.Size.Y * VolumeData.Size.Z * VolumeData.DistanceFieldVolume.GetTypeSize()) / (1024.0f * 1024.0f);
+
+				FNumberFormattingOptions NumberOptions;
+				NumberOptions.MinimumFractionalDigits = 2;
+				NumberOptions.MaximumFractionalDigits = 2;
+
+				if (VolumeData.bMeshWasClosed)
+				{
+					TextItems.Add(SStaticMeshEditorViewport::FOverlayTextItem(
+						FText::Format(NSLOCTEXT("UnrealEd", "DistanceFieldRes_F", "Distance Field:  {0}x{1}x{2} = {3}Mb"), FText::AsNumber(VolumeData.Size.X), FText::AsNumber(VolumeData.Size.Y), FText::AsNumber(VolumeData.Size.Z), FText::AsNumber(MemoryMb, &NumberOptions))));
+				}
+				else
+				{
+					TextItems.Add(SStaticMeshEditorViewport::FOverlayTextItem(
+						NSLOCTEXT("UnrealEd", "DistanceFieldClosed_F", "Distance Field:  Mesh was not closed and material was one-sided")));
+				}
 			}
 		}
 	}
