Generate lightmap UVs on mesh import.

Moved layout UVs to static mesh build process.

diff --git a/Engine/Source/Developer/MeshUtilities/Private/LayoutUV.cpp b/Engine/Source/Developer/MeshUtilities/Private/LayoutUV.cpp
--- a/Engine/Source/Developer/MeshUtilities/Private/LayoutUV.cpp
+++ b/Engine/Source/Developer/MeshUtilities/Private/LayoutUV.cpp
@@ -19,17 +19,10 @@ FLayoutUV::FLayoutUV( FRawMesh* InMesh, uint32 InSrcChannel, uint32 InDstChannel
 	, ChartShader( &ChartRaster )
 {}
 
-void FindOverlappingCorners( TMultiMap<int32,int32>& OutOverlappingCorners, FRawMesh const& RawMesh, float ComparisonThreshold );
-
-void FLayoutUV::FindCharts()
+void FLayoutUV::FindCharts( const TMultiMap<int32,int32>& OverlappingCorners )
 {
 	double Begin = FPlatformTime::Seconds();
 
-	TMultiMap<int32,int32> OverlappingCorners;
-
-	// Find overlapping corners to accelerate adjacency.
-	FindOverlappingCorners( OverlappingCorners, *RawMesh, THRESH_POINTS_ARE_SAME );
-
 	uint32 NumIndexes = RawMesh->WedgeIndices.Num();
 	uint32 NumTris = NumIndexes / 3;
 
diff --git a/Engine/Source/Developer/MeshUtilities/Private/LayoutUV.h b/Engine/Source/Developer/MeshUtilities/Private/LayoutUV.h
--- a/Engine/Source/Developer/MeshUtilities/Private/LayoutUV.h
+++ b/Engine/Source/Developer/MeshUtilities/Private/LayoutUV.h
@@ -42,7 +42,7 @@ class FLayoutUV
 public:
 				FLayoutUV( FRawMesh* InMesh, uint32 InSrcChannel, uint32 InDstChannel, uint32 InTextureResolution );
 
-	void		FindCharts();
+	void		FindCharts( const TMultiMap<int32,int32>& OverlappingCorners );
 	bool		FindBestPacking();
 	void		CommitPackedUVs();
 
diff --git a/Engine/Source/Developer/MeshUtilities/Private/MeshUtilities.cpp b/Engine/Source/Developer/MeshUtilities/Private/MeshUtilities.cpp
--- a/Engine/Source/Developer/MeshUtilities/Private/MeshUtilities.cpp
+++ b/Engine/Source/Developer/MeshUtilities/Private/MeshUtilities.cpp
@@ -1459,7 +1459,7 @@ static void ComputeTriangleTangents(
  * @param OutOverlappingCorners - Maps a corner index to the indices of all overlapping corners.
  * @param RawMesh - The mesh for which to compute overlapping corners.
  */
-void FindOverlappingCorners(
+static void FindOverlappingCorners(
 	TMultiMap<int32,int32>& OutOverlappingCorners,
 	FRawMesh const& RawMesh,
 	float ComparisonThreshold
@@ -2278,6 +2278,19 @@ bool FMeshUtilities::BuildStaticMesh(
 			check(RawMesh.WedgeTangentX.Num() == NumWedges);
 			check(RawMesh.WedgeTangentY.Num() == NumWedges);
 			check(RawMesh.WedgeTangentZ.Num() == NumWedges);
+
+			// Generate lightmap UVs
+			if( SrcModel.BuildSettings.bGenerateLightmapUVs )
+			{
+				FLayoutUV Packer( &RawMesh, SrcModel.BuildSettings.SrcLightmapIndex, SrcModel.BuildSettings.DstLightmapIndex, SrcModel.BuildSettings.MinLightmapResolution );
+
+				Packer.FindCharts( OverlappingCorners );
+				bool bPackSuccess = Packer.FindBestPacking();
+				if( bPackSuccess )
+				{
+					Packer.CommitPackedUVs();
+				}
+			}
 		}
 		else if (LODIndex > 0 && MeshReduction)
 		{
@@ -3155,7 +3168,7 @@ void FMeshUtilities::CreateProxyMesh(
 		StaticMesh->LightingGuid = FGuid::NewGuid();
 
 		// Set it to use textured lightmaps. Note that Build Lighting will do the error-checking (texcoordindex exists for all LODs, etc).
-		StaticMesh->LightMapResolution = 32;
+		StaticMesh->LightMapResolution = 64;
 		StaticMesh->LightMapCoordinateIndex = 1;
 
 		FStaticMeshSourceModel* SrcModel = new (StaticMesh->SourceModels) FStaticMeshSourceModel();
@@ -3794,8 +3807,12 @@ bool FMeshUtilities::LayoutUVs(FRawMesh& RawMesh, uint32 TextureResolution, uint
 	}*/
 
 	FLayoutUV Packer( &RawMesh, 0, TexCoordIndex, TextureResolution );
+	
+	// Find overlapping corners to accelerate adjacency.
+	TMultiMap<int32,int32> OverlappingCorners;
+	FindOverlappingCorners( OverlappingCorners, RawMesh, THRESH_POINTS_ARE_SAME );
 
-	Packer.FindCharts();
+	Packer.FindCharts( OverlappingCorners );
 	bool bPackSuccess = Packer.FindBestPacking();
 	if( bPackSuccess )
 	{
diff --git a/Engine/Source/Editor/StaticMeshEditor/Private/StaticMeshEditorTools.cpp b/Engine/Source/Editor/StaticMeshEditor/Private/StaticMeshEditorTools.cpp
--- a/Engine/Source/Editor/StaticMeshEditor/Private/StaticMeshEditorTools.cpp
+++ b/Engine/Source/Editor/StaticMeshEditor/Private/StaticMeshEditorTools.cpp
@@ -378,6 +378,79 @@ void FMeshBuildSettingsLayout::GenerateChildContent( IDetailChildrenBuilder& Chi
 		];
 	}
 
+	{
+		ChildrenBuilder.AddChildContent( LOCTEXT("GenerateLightmapUVs", "Generate Lightmap UVs").ToString() )
+		.NameContent()
+		[
+			SNew(STextBlock)
+			.Font( IDetailLayoutBuilder::GetDetailFont() )
+			.Text(LOCTEXT("GenerateLightmapUVs", "Generate Lightmap UVs"))
+		]
+		.ValueContent()
+		[
+			SNew(SCheckBox)
+			.IsChecked(this, &FMeshBuildSettingsLayout::ShouldGenerateLightmapUVs)
+			.OnCheckStateChanged(this, &FMeshBuildSettingsLayout::OnGenerateLightmapUVsChanged)
+		];
+	}
+
+	{
+		ChildrenBuilder.AddChildContent( LOCTEXT("MinLightmapResolution", "Min Lightmap Resolution").ToString() )
+		.NameContent()
+		[
+			SNew(STextBlock)
+			.Font( IDetailLayoutBuilder::GetDetailFont() )
+			.Text(LOCTEXT("MinLightmapResolution", "Min Lightmap Resolution").ToString())
+		]
+		.ValueContent()
+		[
+			SNew(SSpinBox<int32>)
+			.Font( IDetailLayoutBuilder::GetDetailFont() )
+			.MinValue(1)
+			.MaxValue(2048)
+			.Value(this, &FMeshBuildSettingsLayout::GetMinLightmapResolution)
+			.OnValueChanged(this, &FMeshBuildSettingsLayout::OnMinLightmapResolutionChanged)
+		];
+	}
+
+	{
+		ChildrenBuilder.AddChildContent( LOCTEXT("SourceLightmapIndex", "Source Lightmap Index").ToString() )
+		.NameContent()
+		[
+			SNew(STextBlock)
+			.Font( IDetailLayoutBuilder::GetDetailFont() )
+			.Text(LOCTEXT("SourceLightmapIndex", "Source Lightmap Index").ToString())
+		]
+		.ValueContent()
+		[
+			SNew(SSpinBox<int32>)
+			.Font( IDetailLayoutBuilder::GetDetailFont() )
+			.MinValue(0)
+			.MaxValue(8)
+			.Value(this, &FMeshBuildSettingsLayout::GetSrcLightmapIndex)
+			.OnValueChanged(this, &FMeshBuildSettingsLayout::OnSrcLightmapIndexChanged)
+		];
+	}
+
+	{
+		ChildrenBuilder.AddChildContent( LOCTEXT("DestinationLightmapIndex", "Destination Lightmap Index").ToString() )
+		.NameContent()
+		[
+			SNew(STextBlock)
+			.Font( IDetailLayoutBuilder::GetDetailFont() )
+			.Text(LOCTEXT("DestinationLightmapIndex", "Destination Lightmap Index").ToString())
+		]
+		.ValueContent()
+		[
+			SNew(SSpinBox<int32>)
+			.Font( IDetailLayoutBuilder::GetDetailFont() )
+			.MinValue(0)
+			.MaxValue(8)
+			.Value(this, &FMeshBuildSettingsLayout::GetDstLightmapIndex)
+			.OnValueChanged(this, &FMeshBuildSettingsLayout::OnDstLightmapIndexChanged)
+		];
+	}
+
 	{
 		ChildrenBuilder.AddChildContent(LOCTEXT("BuildScale", "Build Scale").ToString())
 		.NameContent()
@@ -490,11 +563,31 @@ ESlateCheckBoxState::Type FMeshBuildSettingsLayout::ShouldUseFullPrecisionUVs()
 	return BuildSettings.bUseFullPrecisionUVs ? ESlateCheckBoxState::Checked : ESlateCheckBoxState::Unchecked;
 }
 
+ESlateCheckBoxState::Type FMeshBuildSettingsLayout::ShouldGenerateLightmapUVs() const
+{
+	return BuildSettings.bGenerateLightmapUVs ? ESlateCheckBoxState::Checked : ESlateCheckBoxState::Unchecked;
+}
+
 ESlateCheckBoxState::Type FMeshBuildSettingsLayout::ShouldGenerateDistanceFieldAsIfTwoSided() const
 {
 	return BuildSettings.bGenerateDistanceFieldAsIfTwoSided ? ESlateCheckBoxState::Checked : ESlateCheckBoxState::Unchecked;
 }
 
+int32 FMeshBuildSettingsLayout::GetMinLightmapResolution() const
+{
+	return BuildSettings.MinLightmapResolution;
+}
+
+int32 FMeshBuildSettingsLayout::GetSrcLightmapIndex() const
+{
+	return BuildSettings.SrcLightmapIndex;
+}
+
+int32 FMeshBuildSettingsLayout::GetDstLightmapIndex() const
+{
+	return BuildSettings.DstLightmapIndex;
+}
+
 TOptional<float> FMeshBuildSettingsLayout::GetBuildScaleX() const
 {
 	return BuildSettings.BuildScale3D.X;
@@ -567,6 +660,19 @@ void FMeshBuildSettingsLayout::OnUseFullPrecisionUVsChanged(ESlateCheckBoxState:
 	}
 }
 
+void FMeshBuildSettingsLayout::OnGenerateLightmapUVsChanged(ESlateCheckBoxState::Type NewState)
+{
+	const bool bGenerateLightmapUVs = (NewState == ESlateCheckBoxState::Checked) ? true : false;
+	if (BuildSettings.bGenerateLightmapUVs != bGenerateLightmapUVs)
+	{
+		if (FEngineAnalytics::IsAvailable())
+		{
+			FEngineAnalytics::GetProvider().RecordEvent(TEXT("Editor.Usage.StaticMesh.BuildSettings"), TEXT("bGenerateLightmapUVs"), bGenerateLightmapUVs ? TEXT("True") : TEXT("False"));
+		}
+		BuildSettings.bGenerateLightmapUVs = bGenerateLightmapUVs;
+	}
+}
+
 void FMeshBuildSettingsLayout::OnGenerateDistanceFieldAsIfTwoSidedChanged(ESlateCheckBoxState::Type NewState)
 {
 	const bool bGenerateDistanceFieldAsIfTwoSided = (NewState == ESlateCheckBoxState::Checked) ? true : false;
@@ -580,6 +686,42 @@ void FMeshBuildSettingsLayout::OnGenerateDistanceFieldAsIfTwoSidedChanged(ESlate
 	}
 }
 
+void FMeshBuildSettingsLayout::OnMinLightmapResolutionChanged( int32 NewValue )
+{
+	if (BuildSettings.MinLightmapResolution != NewValue)
+	{
+		if (FEngineAnalytics::IsAvailable())
+		{
+			FEngineAnalytics::GetProvider().RecordEvent(TEXT("Editor.Usage.StaticMesh.BuildSettings"), TEXT("MinLightmapResolution"), FString::Printf(TEXT("%i"), NewValue));
+		}
+		BuildSettings.MinLightmapResolution = NewValue;
+	}
+}
+
+void FMeshBuildSettingsLayout::OnSrcLightmapIndexChanged( int32 NewValue )
+{
+	if (BuildSettings.SrcLightmapIndex != NewValue)
+	{
+		if (FEngineAnalytics::IsAvailable())
+		{
+			FEngineAnalytics::GetProvider().RecordEvent(TEXT("Editor.Usage.StaticMesh.BuildSettings"), TEXT("SrcLightmapIndex"), FString::Printf(TEXT("%i"), NewValue));
+		}
+		BuildSettings.SrcLightmapIndex = NewValue;
+	}
+}
+
+void FMeshBuildSettingsLayout::OnDstLightmapIndexChanged( int32 NewValue )
+{
+	if (BuildSettings.DstLightmapIndex != NewValue)
+	{
+		if (FEngineAnalytics::IsAvailable())
+		{
+			FEngineAnalytics::GetProvider().RecordEvent(TEXT("Editor.Usage.StaticMesh.BuildSettings"), TEXT("DstLightmapIndex"), FString::Printf(TEXT("%i"), NewValue));
+		}
+		BuildSettings.DstLightmapIndex = NewValue;
+	}
+}
+
 void FMeshBuildSettingsLayout::OnBuildScaleXChanged( float NewScaleX, ETextCommit::Type TextCommitType )
 {
 	if (!FMath::IsNearlyEqual(NewScaleX, 0.0f) && BuildSettings.BuildScale3D.X != NewScaleX)
diff --git a/Engine/Source/Editor/StaticMeshEditor/Private/StaticMeshEditorTools.h b/Engine/Source/Editor/StaticMeshEditor/Private/StaticMeshEditorTools.h
--- a/Engine/Source/Editor/StaticMeshEditor/Private/StaticMeshEditorTools.h
+++ b/Engine/Source/Editor/StaticMeshEditor/Private/StaticMeshEditorTools.h
@@ -143,7 +143,11 @@ private:
 	ESlateCheckBoxState::Type ShouldRecomputeTangents() const;
 	ESlateCheckBoxState::Type ShouldRemoveDegenerates() const;
 	ESlateCheckBoxState::Type ShouldUseFullPrecisionUVs() const;
+	ESlateCheckBoxState::Type ShouldGenerateLightmapUVs() const;
 	ESlateCheckBoxState::Type ShouldGenerateDistanceFieldAsIfTwoSided() const;
+	int32 GetMinLightmapResolution() const;
+	int32 GetSrcLightmapIndex() const;
+	int32 GetDstLightmapIndex() const;
 	TOptional<float> GetBuildScaleX() const;
 	TOptional<float> GetBuildScaleY() const;
 	TOptional<float> GetBuildScaleZ() const;
@@ -153,7 +157,11 @@ private:
 	void OnRecomputeTangentsChanged(ESlateCheckBoxState::Type NewState);
 	void OnRemoveDegeneratesChanged(ESlateCheckBoxState::Type NewState);
 	void OnUseFullPrecisionUVsChanged(ESlateCheckBoxState::Type NewState);
+	void OnGenerateLightmapUVsChanged(ESlateCheckBoxState::Type NewState);
 	void OnGenerateDistanceFieldAsIfTwoSidedChanged(ESlateCheckBoxState::Type NewState);
+	void OnMinLightmapResolutionChanged( int32 NewValue );
+	void OnSrcLightmapIndexChanged( int32 NewValue );
+	void OnDstLightmapIndexChanged( int32 NewValue );
 	void OnBuildScaleXChanged( float NewScaleX, ETextCommit::Type TextCommitType );
 	void OnBuildScaleYChanged( float NewScaleY, ETextCommit::Type TextCommitType );
 	void OnBuildScaleZChanged( float NewScaleZ, ETextCommit::Type TextCommitType );
diff --git a/Engine/Source/Editor/UnrealEd/Classes/Factories/FbxStaticMeshImportData.h b/Engine/Source/Editor/UnrealEd/Classes/Factories/FbxStaticMeshImportData.h
--- a/Engine/Source/Editor/UnrealEd/Classes/Factories/FbxStaticMeshImportData.h
+++ b/Engine/Source/Editor/UnrealEd/Classes/Factories/FbxStaticMeshImportData.h
@@ -23,6 +23,9 @@ class UFbxStaticMeshImportData : public UFbxMeshImportData
 	/** Disabling this option will keep degenerate triangles found.  In general you should leave this option on. */
 	UPROPERTY(EditAnywhere, Category=ImportSettings)
 	uint32 bRemoveDegenerates:1;
+	
+	UPROPERTY(EditAnywhere, Category=ImportSettings)
+	uint32 bGenerateLightmapUVs:1;
 
 	/** If checked, one convex hull per UCX_ prefixed collision mesh will be generated instead of decomposing into multiple hulls */
 	UPROPERTY(EditAnywhere, config, Category=ImportSettings)
diff --git a/Engine/Source/Editor/UnrealEd/Private/Fbx/FbxMainImport.cpp b/Engine/Source/Editor/UnrealEd/Private/Fbx/FbxMainImport.cpp
--- a/Engine/Source/Editor/UnrealEd/Private/Fbx/FbxMainImport.cpp
+++ b/Engine/Source/Editor/UnrealEd/Private/Fbx/FbxMainImport.cpp
@@ -195,6 +195,7 @@ void ApplyImportUIToImportOptions(UFbxImportUI* ImportUI, FBXImportOptions& InOu
 	InOutImportOptions.bCombineToSingle = ImportUI->bCombineMeshes;
 	InOutImportOptions.bReplaceVertexColors = ImportUI->StaticMeshImportData->bReplaceVertexColors;
 	InOutImportOptions.bRemoveDegenerates = ImportUI->StaticMeshImportData->bRemoveDegenerates;
+	InOutImportOptions.bGenerateLightmapUVs = ImportUI->StaticMeshImportData->bGenerateLightmapUVs;
 	InOutImportOptions.bOneConvexHullPerUCX = ImportUI->StaticMeshImportData->bOneConvexHullPerUCX;
 	InOutImportOptions.StaticMeshLODGroup = ImportUI->StaticMeshImportData->StaticMeshLODGroup;
 	InOutImportOptions.bImportMeshesInBoneHierarchy = ImportUI->SkeletalMeshImportData->bImportMeshesInBoneHierarchy;
diff --git a/Engine/Source/Editor/UnrealEd/Private/Fbx/FbxOptionWindow.cpp b/Engine/Source/Editor/UnrealEd/Private/Fbx/FbxOptionWindow.cpp
--- a/Engine/Source/Editor/UnrealEd/Private/Fbx/FbxOptionWindow.cpp
+++ b/Engine/Source/Editor/UnrealEd/Private/Fbx/FbxOptionWindow.cpp
@@ -722,6 +722,17 @@ TSharedRef<SWidget> SFbxOptionWindow::ConstructStaticMeshAdvanced()
 		]
 	];
 
+	NewBox->AddSlot().AutoHeight() .Padding(2)
+	[
+		SNew(SCheckBox)
+		.IsChecked(ImportUI->StaticMeshImportData->bGenerateLightmapUVs ? ESlateCheckBoxState::Checked : ESlateCheckBoxState::Unchecked)
+		.OnCheckStateChanged(this, &SFbxOptionWindow::SetStaticMesh_GenerateLightmapUVs)
+		[
+			SNew(STextBlock)
+			.Text(LOCTEXT("FbxOptionWindow_GenerateLightmapUVs", "Generate Lightmap UVs"))
+		]
+	];
+
 	NewBox->AddSlot().AutoHeight() .Padding(2)
 	[
 		SNew(SCheckBox)
@@ -1247,6 +1258,11 @@ void SFbxOptionWindow::SetStaticMesh_RemoveDegenerates(ESlateCheckBoxState::Type
 	ImportUI->StaticMeshImportData->bRemoveDegenerates = (NewType == ESlateCheckBoxState::Checked)? 1:0;
 }
 
+void SFbxOptionWindow::SetStaticMesh_GenerateLightmapUVs(ESlateCheckBoxState::Type NewType)
+{
+	ImportUI->StaticMeshImportData->bGenerateLightmapUVs = (NewType == ESlateCheckBoxState::Checked)? 1:0;
+}
+
 void SFbxOptionWindow::SetStaticMesh_OneConvexHullPerUCX(ESlateCheckBoxState::Type NewType)
 {
 	ImportUI->StaticMeshImportData->bOneConvexHullPerUCX = (NewType == ESlateCheckBoxState::Checked)? 1:0;
diff --git a/Engine/Source/Editor/UnrealEd/Private/Fbx/FbxStaticMeshImport.cpp b/Engine/Source/Editor/UnrealEd/Private/Fbx/FbxStaticMeshImport.cpp
--- a/Engine/Source/Editor/UnrealEd/Private/Fbx/FbxStaticMeshImport.cpp
+++ b/Engine/Source/Editor/UnrealEd/Private/Fbx/FbxStaticMeshImport.cpp
@@ -935,7 +935,7 @@ UStaticMesh* UnFbx::FFbxImporter::ImportStaticMeshAsSingle(UObject* InParent, TA
 	StaticMesh->LightingGuid = FGuid::NewGuid();
 
 	// Set it to use textured lightmaps. Note that Build Lighting will do the error-checking (texcoordindex exists for all LODs, etc).
-	StaticMesh->LightMapResolution = 32;
+	StaticMesh->LightMapResolution = 64;
 	StaticMesh->LightMapCoordinateIndex = 1;
 
 	TArray<FFbxMaterial> MeshMaterials;
@@ -1157,6 +1157,7 @@ UStaticMesh* UnFbx::FFbxImporter::ImportStaticMeshAsSingle(UObject* InParent, TA
 		SrcModel.BuildSettings.bRemoveDegenerates = ImportOptions->bRemoveDegenerates;
 		SrcModel.BuildSettings.bRecomputeNormals = ImportOptions->NormalImportMethod == FBXNIM_ComputeNormals;
 		SrcModel.BuildSettings.bRecomputeTangents = ImportOptions->NormalImportMethod != FBXNIM_ImportNormalsAndTangents;
+		SrcModel.BuildSettings.bGenerateLightmapUVs = ImportOptions->bGenerateLightmapUVs;
 
 		StaticMesh->LODGroup = ImportOptions->StaticMeshLODGroup;
 		StaticMesh->Build(false);
diff --git a/Engine/Source/Editor/UnrealEd/Private/Fbx/FbxStaticMeshImportData.cpp b/Engine/Source/Editor/UnrealEd/Private/Fbx/FbxStaticMeshImportData.cpp
--- a/Engine/Source/Editor/UnrealEd/Private/Fbx/FbxStaticMeshImportData.cpp
+++ b/Engine/Source/Editor/UnrealEd/Private/Fbx/FbxStaticMeshImportData.cpp
@@ -7,6 +7,7 @@ UFbxStaticMeshImportData::UFbxStaticMeshImportData(const class FPostConstructIni
 {
 	StaticMeshLODGroup = NAME_None;
 	bRemoveDegenerates = true;
+	bGenerateLightmapUVs = true;
 	bOneConvexHullPerUCX = true;
 }
 
diff --git a/Engine/Source/Editor/UnrealEd/Private/FbxImporter.h b/Engine/Source/Editor/UnrealEd/Private/FbxImporter.h
--- a/Engine/Source/Editor/UnrealEd/Private/FbxImporter.h
+++ b/Engine/Source/Editor/UnrealEd/Private/FbxImporter.h
@@ -113,6 +113,7 @@ struct FBXImportOptions
 	bool bCombineToSingle;
 	bool bReplaceVertexColors;
 	bool bRemoveDegenerates;
+	bool bGenerateLightmapUVs;
 	bool bOneConvexHullPerUCX;
 	FName StaticMeshLODGroup;
 	// Skeletal Mesh options
diff --git a/Engine/Source/Editor/UnrealEd/Private/FbxOptionWindow.h b/Engine/Source/Editor/UnrealEd/Private/FbxOptionWindow.h
--- a/Engine/Source/Editor/UnrealEd/Private/FbxOptionWindow.h
+++ b/Engine/Source/Editor/UnrealEd/Private/FbxOptionWindow.h
@@ -116,6 +116,7 @@ private:
 	void SetStaticMesh_CombineMeshes(ESlateCheckBoxState::Type NewType);
 	void SetStaticMesh_ReplaceVertexColor(ESlateCheckBoxState::Type NewType);
 	void SetStaticMesh_RemoveDegenerates(ESlateCheckBoxState::Type NewType);
+	void SetStaticMesh_GenerateLightmapUVs(ESlateCheckBoxState::Type NewType);
 	void SetStaticMesh_OneConvexHullPerUCX(ESlateCheckBoxState::Type NewType);
 	void SetMaterial_ImportMaterials(ESlateCheckBoxState::Type NewType);
 	void SetMaterial_ImportTextures(ESlateCheckBoxState::Type NewType);
diff --git a/Engine/Source/Editor/WorldBrowser/Private/Tiles/WorldTileCollectionModel.cpp b/Engine/Source/Editor/WorldBrowser/Private/Tiles/WorldTileCollectionModel.cpp
--- a/Engine/Source/Editor/WorldBrowser/Private/Tiles/WorldTileCollectionModel.cpp
+++ b/Engine/Source/Editor/WorldBrowser/Private/Tiles/WorldTileCollectionModel.cpp
@@ -1940,7 +1940,7 @@ bool FWorldTileCollectionModel::GenerateLODLevels(FLevelModelList InLevelList, i
 				StaticMesh->LightingGuid = FGuid::NewGuid();
 
 				// Set it to use textured lightmaps. Note that Build Lighting will do the error-checking (texcoordindex exists for all LODs, etc).
-				StaticMesh->LightMapResolution = 128;
+				StaticMesh->LightMapResolution = 64;
 				StaticMesh->LightMapCoordinateIndex = 1;
 
 				FStaticMeshSourceModel* SrcModel = new (StaticMesh->SourceModels) FStaticMeshSourceModel();
diff --git a/Engine/Source/Runtime/Core/Public/UObject/ObjectVersion.h b/Engine/Source/Runtime/Core/Public/UObject/ObjectVersion.h
--- a/Engine/Source/Runtime/Core/Public/UObject/ObjectVersion.h
+++ b/Engine/Source/Runtime/Core/Public/UObject/ObjectVersion.h
@@ -602,6 +602,8 @@ enum EUnrealEngineObjectUE4Version
 	VER_UE4_AUTO_WELDING,
 	// Rename UCharacterMovementComponent::bCrouchMovesCharacterDown
 	VER_UE4_RENAME_CROUCHMOVESCHARACTERDOWN,
+	// Lightmap parameters in FMeshBuildSettings
+	VER_UE4_LIGHTMAP_MESH_BUILD_SETTINGS,
 
 	// -----<new versions can be added before this line>-------------------------------------------------
 	// - this needs to be the last line (see note below)
diff --git a/Engine/Source/Runtime/Engine/Classes/Engine/EngineTypes.h b/Engine/Source/Runtime/Engine/Classes/Engine/EngineTypes.h
--- a/Engine/Source/Runtime/Engine/Classes/Engine/EngineTypes.h
+++ b/Engine/Source/Runtime/Engine/Classes/Engine/EngineTypes.h
@@ -1793,6 +1793,18 @@ struct FMeshBuildSettings
 	UPROPERTY(EditAnywhere, Category=BuildSettings)
 	bool bUseFullPrecisionUVs;
 
+	UPROPERTY(EditAnywhere, Category=BuildSettings)
+	bool bGenerateLightmapUVs;
+
+	UPROPERTY(EditAnywhere, Category=BuildSettings)
+	int32 MinLightmapResolution;
+
+	UPROPERTY(EditAnywhere, Category=BuildSettings, meta=(DisplayName="Source Lightmap Index"))
+	int32 SrcLightmapIndex;
+
+	UPROPERTY(EditAnywhere, Category=BuildSettings, meta=(DisplayName="Destination Lightmap Index"))
+	int32 DstLightmapIndex;
+
 	UPROPERTY()
 	float BuildScale_DEPRECATED;
 
@@ -1808,7 +1820,7 @@ struct FMeshBuildSettings
 	float DistanceFieldResolutionScale;
 
 	/** 
-	 * Whether to generate the distance field treating every triangle hit as a frontface.  
+	 * Whether to generate the distance field treating every triangle hit as a front face.  
 	 * When enabled prevents the distance field from being discarded due to the mesh being open, but also lowers Distance Field AO quality.
 	 */
 	UPROPERTY(EditAnywhere, Category=BuildSettings)
@@ -1820,12 +1832,15 @@ struct FMeshBuildSettings
 		, bRecomputeTangents(true)
 		, bRemoveDegenerates(true)
 		, bUseFullPrecisionUVs(false)
+		, bGenerateLightmapUVs(true)
+		, MinLightmapResolution(64)
+		, SrcLightmapIndex(0)
+		, DstLightmapIndex(1)
 		, BuildScale_DEPRECATED(1.0f)
 		, BuildScale3D(1.0f, 1.0f, 1.0f)
-	{
-		DistanceFieldResolutionScale = 1;
-		bGenerateDistanceFieldAsIfTwoSided = false;
-	}
+		, DistanceFieldResolutionScale(1.0f)
+		, bGenerateDistanceFieldAsIfTwoSided(false)
+	{}
 
 	/** Equality operator. */
 	bool operator==(const FMeshBuildSettings& Other) const
@@ -1834,6 +1849,10 @@ struct FMeshBuildSettings
 			&& bRecomputeTangents == Other.bRecomputeTangents
 			&& bRemoveDegenerates == Other.bRemoveDegenerates
 			&& bUseFullPrecisionUVs == Other.bUseFullPrecisionUVs
+			&& bGenerateLightmapUVs == Other.bGenerateLightmapUVs
+			&& MinLightmapResolution == Other.MinLightmapResolution
+			&& SrcLightmapIndex == Other.SrcLightmapIndex
+			&& DstLightmapIndex == Other.DstLightmapIndex
 			&& BuildScale3D == Other.BuildScale3D
 			&& DistanceFieldResolutionScale == Other.DistanceFieldResolutionScale
 			&& bGenerateDistanceFieldAsIfTwoSided == Other.bGenerateDistanceFieldAsIfTwoSided;
diff --git a/Engine/Source/Runtime/Engine/Private/StaticMesh.cpp b/Engine/Source/Runtime/Engine/Private/StaticMesh.cpp
--- a/Engine/Source/Runtime/Engine/Private/StaticMesh.cpp
+++ b/Engine/Source/Runtime/Engine/Private/StaticMesh.cpp
@@ -985,6 +985,11 @@ FArchive& operator<<(FArchive& Ar, FMeshBuildSettings& BuildSettings)
 	Ar << BuildSettings.bRecomputeTangents;
 	Ar << BuildSettings.bRemoveDegenerates;
 	Ar << BuildSettings.bUseFullPrecisionUVs;
+	Ar << BuildSettings.bGenerateLightmapUVs;
+
+	Ar << BuildSettings.MinLightmapResolution;
+	Ar << BuildSettings.SrcLightmapIndex;
+	Ar << BuildSettings.DstLightmapIndex;
 
 	if (Ar.IsLoading() && Ar.UE4Ver() < VER_UE4_BUILD_SCALE_VECTOR)
 	{
@@ -1006,7 +1011,7 @@ FArchive& operator<<(FArchive& Ar, FMeshBuildSettings& BuildSettings)
 // differences, etc.) replace the version GUID below with a new one.
 // In case of merge conflicts with DDC versions, you MUST generate a new GUID
 // and set this new GUID as the version.
-#define STATICMESH_DERIVEDDATA_VER TEXT("4668778661B445A9523C94440EA899D")
+#define STATICMESH_DERIVEDDATA_VER TEXT("46A8778361B442A9523C54440EA1E9D")
 
 static const FString& GetStaticMeshDerivedDataVersion()
 {
@@ -1848,6 +1853,14 @@ void UStaticMesh::PostLoad()
 		}
 	}
 
+	if( GetLinkerUE4Version() < VER_UE4_LIGHTMAP_MESH_BUILD_SETTINGS )
+	{
+		for( int32 i = 0; i < SourceModels.Num(); i++ )
+		{
+			SourceModels[i].BuildSettings.bGenerateLightmapUVs = false;
+		}
+	}
+
 	CacheDerivedData();
 
 	// Only required in an editor build as other builds process this in a different place
diff --git a/Engine/Source/Runtime/Engine/Public/StaticMeshResources.h b/Engine/Source/Runtime/Engine/Public/StaticMeshResources.h
--- a/Engine/Source/Runtime/Engine/Public/StaticMeshResources.h
+++ b/Engine/Source/Runtime/Engine/Public/StaticMeshResources.h
@@ -19,7 +19,7 @@ public:
 	/** Default values. */
 	FStaticMeshLODGroup()
 		: DefaultNumLODs(1)
-		, DefaultLightMapResolution(32)
+		, DefaultLightMapResolution(64)
 		, BasePercentTrianglesMult(1.0f)
 		, DisplayName( NSLOCTEXT( "UnrealEd", "None", "None" ) )
 	{
