Tore out D3DX generate UVs and the mesh editor tab for it.

diff --git a/Engine/Source/Developer/MeshUtilities/Private/MeshUtilities.cpp b/Engine/Source/Developer/MeshUtilities/Private/MeshUtilities.cpp
--- a/Engine/Source/Developer/MeshUtilities/Private/MeshUtilities.cpp
+++ b/Engine/Source/Developer/MeshUtilities/Private/MeshUtilities.cpp
@@ -90,19 +90,6 @@ private:
 
 	virtual bool BuildSkeletalMesh( FStaticLODModel& LODModel, const FReferenceSkeleton& RefSkeleton, const TArray<FVertInfluence>& Influences, const TArray<FMeshWedge>& Wedges, const TArray<FMeshFace>& Faces, const TArray<FVector>& Points, const TArray<int32>& PointToOriginalMap, bool bKeepOverlappingVertices = false, bool bComputeNormals = true, bool bComputeTangents = true, TArray<FText> * OutWarningMessages = NULL, TArray<FName> * OutWarningNames = NULL);
 
-	virtual bool GenerateUVs(
-		FRawMesh& RawMesh,
-		uint32 TexCoordIndex,
-		float MinChartSpacingPercent,
-		float BorderSpacingPercent,
-		bool bUseMaxStretch,
-		const TArray< int32 >* InFalseEdgeIndices,
-		uint32& MaxCharts,
-		float& MaxDesiredStretch,
-		FText& OutError
-		) override;
-
-	virtual bool LayoutUVs(FRawMesh& RawMesh, uint32 TextureResolution, uint32 TexCoordIndex, FText& OutError) override;
 	virtual IMeshReduction* GetMeshReductionInterface() override;
 	virtual IMeshMerging* GetMeshMergingInterface() override;
 	virtual void CacheOptimizeIndexBuffer(TArray<uint16>& Indices) override;
@@ -3746,83 +3733,6 @@ void FMeshUtilities::MergeActors(
 }
 	
 
-/*------------------------------------------------------------------------------
-	UV Generation
-------------------------------------------------------------------------------*/
-
-#if PLATFORM_WINDOWS
-#include "Windows/D3D9MeshUtils.h"
-#endif // #if PLATFORM_WINDOWS
-
-bool FMeshUtilities::GenerateUVs(
-	FRawMesh& RawMesh,
-	uint32 TexCoordIndex,
-	float MinChartSpacingPercent,
-	float BorderSpacingPercent,
-	bool bUseMaxStretch,
-	const TArray< int32 >* InFalseEdgeIndices,
-	uint32& MaxCharts,
-	float& MaxDesiredStretch,
-	FText& OutError
-	)
-{
-#if PLATFORM_WINDOWS
-	FD3D9MeshUtilities D3DMeshUtils;
-	return D3DMeshUtils.GenerateUVs(RawMesh, TexCoordIndex, MinChartSpacingPercent, BorderSpacingPercent, bUseMaxStretch, InFalseEdgeIndices, MaxCharts, MaxDesiredStretch, OutError);
-#else
-	return false;
-#endif // #if PLATFORM_WINDOWS
-}
-
-
-
-bool FMeshUtilities::LayoutUVs(FRawMesh& RawMesh, uint32 TextureResolution, uint32 TexCoordIndex, FText& OutError)
-{
-#if 0//PLATFORM_WINDOWS
-	FD3D9MeshUtilities D3DMeshUtils;
-	return D3DMeshUtils.LayoutUVs(RawMesh, TextureResolution, TexCoordIndex, OutError);
-#else
-
-	OutError = FText();
-	if( !RawMesh.IsValid() )
-	{
-		OutError = LOCTEXT("LayoutUVs_FailedInvalid", "LayoutUVs failed, mesh was invalid.");
-		return false;
-	}
-
-	/*int32 NumTexCoords = 0;
-	for (int32 i = 0; i < MAX_MESH_TEXTURE_COORDS; ++i)
-	{
-		if (RawMesh.WedgeTexCoords[i].Num() != RawMesh.WedgeIndices.Num())
-		{
-			break;
-		}
-		NumTexCoords++;
-	}
-
-	if (NumTexCoords > MAX_MESH_TEXTURE_COORDS || TexCoordIndex > (uint32)NumTexCoords)
-	{
-		OutError = LOCTEXT("LayoutUVs_FailedUVs", "LayoutUVs failed, incorrect number of texcoords.");
-	return false;
-	}*/
-
-	FLayoutUV Packer( &RawMesh, 0, TexCoordIndex, TextureResolution );
-	
-	// Find overlapping corners to accelerate adjacency.
-	TMultiMap<int32,int32> OverlappingCorners;
-	FindOverlappingCorners( OverlappingCorners, RawMesh, THRESH_POINTS_ARE_SAME );
-
-	Packer.FindCharts( OverlappingCorners );
-	bool bPackSuccess = Packer.FindBestPacking();
-	if( bPackSuccess )
-	{
-		Packer.CommitPackedUVs();
-	}
-
-	return bPackSuccess;
-#endif // #if PLATFORM_WINDOWS
-}
-
 /*------------------------------------------------------------------------------
 	Mesh reduction.
 ------------------------------------------------------------------------------*/
diff --git a/Engine/Source/Developer/MeshUtilities/Private/Windows/D3D9MeshUtils.cpp b/Engine/Source/Developer/MeshUtilities/Private/Windows/D3D9MeshUtils.cpp
deleted file mode 100644
diff --git a/Engine/Source/Developer/MeshUtilities/Private/Windows/D3D9MeshUtils.h b/Engine/Source/Developer/MeshUtilities/Private/Windows/D3D9MeshUtils.h
deleted file mode 100644
diff --git a/Engine/Source/Developer/MeshUtilities/Public/MeshUtilities.h b/Engine/Source/Developer/MeshUtilities/Public/MeshUtilities.h
--- a/Engine/Source/Developer/MeshUtilities/Public/MeshUtilities.h
+++ b/Engine/Source/Developer/MeshUtilities/Public/MeshUtilities.h
@@ -163,49 +163,6 @@ public:
 		TArray<FName> * OutWarningNames = NULL
 		) = 0;
 
-
-	/**
-	 * Generates a unique UV layout for a static mesh
-	 * @param RawMesh - The input/output mesh
-	 * @param TexCoordIndex - Index of the uv channel to overwrite or create
-	 * @param bKeepExistingCoordinates - True to preserve the existing charts when packing
-	 * @param MinChartSpacingPercent - Minimum distance between two packed charts (0.0-100.0)
-	 * @param BorderSpacingPercent - Spacing between UV border and charts (0.0-100.0)
-	 * @param bUseMaxStretch - True if "MaxDesiredStretch" should be used; otherwise "MaxCharts" is used
-	 * @param InFalseEdgeIndices - Optional list of raw face edge indices to be ignored when creating UV seams
-	 * @param MaxCharts - In: Max number of charts to allow; Out:Number of charts generated by the uv unwrap algorithm.
-	 * @param MaxDesiredStretch - The amount of stretching allowed. 0 means no stretching is allowed, 1 means any amount of stretching can be used. 
-	 * @param OutError - if unsuccessful, contains the error that occurred.
-	 * @return true if successful
-	 */
-	virtual bool GenerateUVs(
-		struct FRawMesh& RawMesh,
-		uint32 TexCoordIndex,
-		float MinChartSpacingPercent,
-		float BorderSpacingPercent,
-		bool bUseMaxStretch,
-		const TArray< int32 >* InFalseEdgeIndices,
-		uint32& MaxCharts,
-		float& MaxDesiredStretch,
-		FText& OutError
-		) = 0;
-
-	/**
-	 * For quick generating lightmap uvs.
-	 * It copies charts from 0 uv channel and layouts them without making new charts (keeps edge splits). Additionally separates folded triangles automatically
-	 * Use when DXD generates ugly cuts and degenerates charts 
-	 * @param RawMesh - The input/output mesh
-	 * @param TexCoordIndex - Index of the uv channel to overwrite or create
-	 * @param OutError - if unsuccessful, contains the error that occurred.
-	 * @return true if successful
-	 */
-	virtual bool LayoutUVs(
-	struct FRawMesh& RawMesh,
-		uint32 TextureResolution,
-		uint32 TexCoordIndex,
-		FText& OutError
-		) = 0;
-
 	/** Returns the mesh reduction plugin if available. */
 	virtual IMeshReduction* GetMeshReductionInterface() = 0;
 
diff --git a/Engine/Source/Editor/StaticMeshEditor/Private/StaticMeshEditor.cpp b/Engine/Source/Editor/StaticMeshEditor/Private/StaticMeshEditor.cpp
--- a/Engine/Source/Editor/StaticMeshEditor/Private/StaticMeshEditor.cpp
+++ b/Engine/Source/Editor/StaticMeshEditor/Private/StaticMeshEditor.cpp
@@ -33,7 +33,6 @@ const FName FStaticMeshEditor::ViewportTabId( TEXT( "StaticMeshEditor_Viewport"
 const FName FStaticMeshEditor::PropertiesTabId( TEXT( "StaticMeshEditor_Properties" ) );
 const FName FStaticMeshEditor::SocketManagerTabId( TEXT( "StaticMeshEditor_SocketManager" ) );
 const FName FStaticMeshEditor::CollisionTabId( TEXT( "StaticMeshEditor_Collision" ) );
-const FName FStaticMeshEditor::GenerateUniqueUVsTabId( TEXT( "StaticMeshEditor_GenerateUniqueUVs" ) );
 
 
 void FStaticMeshEditor::RegisterTabSpawners(const TSharedRef<class FTabManager>& TabManager)
@@ -57,10 +56,6 @@ void FStaticMeshEditor::RegisterTabSpawners(const TSharedRef<class FTabManager>&
 	TabManager->RegisterTabSpawner( CollisionTabId, FOnSpawnTab::CreateSP(this, &FStaticMeshEditor::SpawnTab_Collision) )
 		.SetDisplayName( LOCTEXT("CollisionTab", "Convex Decomposition") )
 		.SetGroup( MenuStructure.GetAssetEditorCategory() );
-	
-	TabManager->RegisterTabSpawner( GenerateUniqueUVsTabId, FOnSpawnTab::CreateSP(this, &FStaticMeshEditor::SpawnTab_GenerateUniqueUVs) )
-		.SetDisplayName( LOCTEXT("GenerateUniqueUVsTab", "Generate Unique UVs") )
-		.SetGroup( MenuStructure.GetAssetEditorCategory() );
 }
 
 void FStaticMeshEditor::UnregisterTabSpawners(const TSharedRef<class FTabManager>& TabManager)
@@ -71,7 +66,6 @@ void FStaticMeshEditor::UnregisterTabSpawners(const TSharedRef<class FTabManager
 	TabManager->UnregisterTabSpawner( PropertiesTabId );
 	TabManager->UnregisterTabSpawner( SocketManagerTabId );
 	TabManager->UnregisterTabSpawner( CollisionTabId );
-	TabManager->UnregisterTabSpawner( GenerateUniqueUVsTabId );
 }
 
 
@@ -158,7 +152,6 @@ void FStaticMeshEditor::InitStaticMeshEditor( const EToolkitMode::Type Mode, con
 					->SetSizeCoefficient(0.5f)
 					->AddTab(SocketManagerTabId, ETabState::ClosedTab)
 					->AddTab(CollisionTabId, ETabState::ClosedTab)
-					->AddTab(GenerateUniqueUVsTabId, ETabState::ClosedTab)
 				)
 			)
 		)
@@ -338,27 +331,6 @@ TSharedRef<SDockTab> FStaticMeshEditor::SpawnTab_Collision( const FSpawnTabArgs&
 		];
 }
 
-TSharedRef<SDockTab> FStaticMeshEditor::SpawnTab_GenerateUniqueUVs( const FSpawnTabArgs& Args )
-{
-	check( Args.GetTabId() == GenerateUniqueUVsTabId );
-	
-#if !PLATFORM_WINDOWS
-	FText ErrorMessage = FText::Format( LOCTEXT("StaticMeshGenerateUniqueUVs_Unsupported", "Generate Unique UVs not yet implemented for {0}."), FText::FromString( FPlatformProperties::IniPlatformName() ) );
-	OpenMsgDlgInt( EAppMsgType::Ok, ErrorMessage, LOCTEXT("StaticMeshGenerateUniqueUVs_UnsupportedErrorCaption", "Error") );
-#endif
-	
-	return SNew(SDockTab)
-		.Label( LOCTEXT("StaticMeshGenerateUniqueUVs_TabTitle", "Generate Unique UVs") )
-		[
-#if PLATFORM_WINDOWS
-			GenerateUniqueUVs.ToSharedRef()
-#else
-			SNew(STextBlock)
-			.Text( ErrorMessage )
-#endif
-		];
-}
-
 void FStaticMeshEditor::BindCommands()
 {
 	const FStaticMeshEditorCommands& Commands = FStaticMeshEditorCommands::Get();
@@ -518,9 +490,6 @@ void FStaticMeshEditor::BuildSubTools()
 	SAssignNew( ConvexDecomposition, SConvexDecomposition )
 		.StaticMeshEditorPtr(SharedThis(this));
 
-	SAssignNew(GenerateUniqueUVs, SGenerateUniqueUVs)
-		.StaticMeshEditorPtr(SharedThis(this));
-
 	// Build toolbar widgets
 	UVChannelCombo = SNew(STextComboBox)
 		.OptionsSource(&UVChannels)
@@ -1409,12 +1378,6 @@ void FStaticMeshEditor::OnChangeMesh()
 
 		SetEditorMesh(SelectedMesh);
 
-		// Refresh the tool so it will update it's LOD list.
-		if(GenerateUniqueUVs.IsValid())
-		{
-			GenerateUniqueUVs->RefreshTool();
-		}
-
 		if(SocketManager.IsValid())
 		{
 			SocketManager->UpdateStaticMesh();
diff --git a/Engine/Source/Editor/StaticMeshEditor/Private/StaticMeshEditor.h b/Engine/Source/Editor/StaticMeshEditor/Private/StaticMeshEditor.h
--- a/Engine/Source/Editor/StaticMeshEditor/Private/StaticMeshEditor.h
+++ b/Engine/Source/Editor/StaticMeshEditor/Private/StaticMeshEditor.h
@@ -114,7 +114,6 @@ private:
 	TSharedRef<SDockTab> SpawnTab_Properties(const FSpawnTabArgs& Args);
 	TSharedRef<SDockTab> SpawnTab_SocketManager(const FSpawnTabArgs& Args);
 	TSharedRef<SDockTab> SpawnTab_Collision(const FSpawnTabArgs& Args);
-	TSharedRef<SDockTab> SpawnTab_GenerateUniqueUVs(const FSpawnTabArgs& Args);
 
 private:
 	/** Binds commands associated with the Static Mesh Editor. */
@@ -259,9 +258,6 @@ private:
 	/** Convex Decomposition widget */
 	TSharedPtr< class SConvexDecomposition> ConvexDecomposition;
 
-	/** Generate Unique UVs widget. */
-	TSharedPtr< class SGenerateUniqueUVs> GenerateUniqueUVs;
-
 	/** Widget for displaying the available UV Channels. */
 	TSharedPtr< class STextComboBox > UVChannelCombo;
 
@@ -306,5 +302,4 @@ private:
 	static const FName PropertiesTabId;
 	static const FName SocketManagerTabId;
 	static const FName CollisionTabId;
-	static const FName GenerateUniqueUVsTabId;
 };
diff --git a/Engine/Source/Editor/StaticMeshEditor/Private/StaticMeshEditorActions.cpp b/Engine/Source/Editor/StaticMeshEditor/Private/StaticMeshEditorActions.cpp
--- a/Engine/Source/Editor/StaticMeshEditor/Private/StaticMeshEditorActions.cpp
+++ b/Engine/Source/Editor/StaticMeshEditor/Private/StaticMeshEditorActions.cpp
@@ -39,8 +39,6 @@ void FStaticMeshEditorCommands::RegisterCommands()
 
 	UI_COMMAND( FindSource, "Find Source", "Opens explorer at the location of this asset.", EUserInterfaceActionType::Button, FInputGesture() );
 	
-	UI_COMMAND( GenerateUniqueUVs, "Generate Unique UVs...", "Opens the _Generate Unique UVs_ pane for generating a set of unique (non-overlapping) texture coordinates.", EUserInterfaceActionType::Button, FInputGesture() );
-
 	UI_COMMAND( ChangeMesh, "Change Mesh", "Changes the static mesh asset loaded in the Static Mesh Editor to the asset currently selected in the Content Browser.", EUserInterfaceActionType::Button, FInputGesture() );
 }
 
diff --git a/Engine/Source/Editor/StaticMeshEditor/Private/StaticMeshEditorTools.cpp b/Engine/Source/Editor/StaticMeshEditor/Private/StaticMeshEditorTools.cpp
--- a/Engine/Source/Editor/StaticMeshEditor/Private/StaticMeshEditorTools.cpp
+++ b/Engine/Source/Editor/StaticMeshEditor/Private/StaticMeshEditorTools.cpp
@@ -1992,432 +1992,4 @@ FText FLevelOfDetailSettingsLayout::GetLODCountTooltip() const
 	return LOCTEXT("LODCountTooltip_Disabled", "Auto mesh reduction is unavailable! Please provide a mesh reduction interface such as Simplygon to use this feature or manually import LOD levels.");
 }
 
-/////////////////////////////////
-// SGenerateUniqueUVs
-/////////////////////////////////
-BEGIN_SLATE_FUNCTION_BUILD_OPTIMIZATION
-void SGenerateUniqueUVs::Construct(const FArguments& InArgs)
-{
-	StaticMeshEditorPtr = InArgs._StaticMeshEditorPtr;
-
-	RefreshTool();
-
-	MaxStretching = 0.5f;
-	MaxCharts = 100.0f;
-	MinSpacingBetweenCharts = 1.0f;
-
-	this->ChildSlot
-		[
-			SNew(SVerticalBox)
-			
-			+SVerticalBox::Slot()
-			.AutoHeight()
-			[
-				SNew(STextBlock)
-					.Text( LOCTEXT("CreationMode", "Creation Mode") )
-			]
-
-			+SVerticalBox::Slot()
-				.AutoHeight()
-				.HAlign(HAlign_Center)
-			[
-				SNew(SHorizontalBox)
-
-				+SHorizontalBox::Slot()
-					.AutoWidth()
-					.Padding(0.0f, 0.0f, 8.0f, 0.0f)
-				[
-					SNew(SCheckBox)	
-						.Style(FEditorStyle::Get(), "RadioButton")
-						.IsChecked( this, &SGenerateUniqueUVs::IsCreationModeChecked, ECreationModeChoice::CreateNew )
-						.OnCheckStateChanged( this, &SGenerateUniqueUVs::OnCreationModeChanged, CreateNew )
-						[
-							SNew(STextBlock)
-								.Text( LOCTEXT("CreateNew", "Create New") )
-						]
-				]
-
-				+SHorizontalBox::Slot()
-					.AutoWidth()
-					.Padding(8.0f, 0.0f, 0.0f, 0.0f)
-				[
-					SNew(SCheckBox)	
-						.Style(FEditorStyle::Get(), "RadioButton")
-						.IsChecked( this, &SGenerateUniqueUVs::IsCreationModeChecked, ECreationModeChoice::UseChannel0 )
-						.OnCheckStateChanged( this, &SGenerateUniqueUVs::OnCreationModeChanged, UseChannel0 )
-						[
-							SNew(STextBlock)
-								.Text( LOCTEXT("LayoutUsing0Channel", "Layout using 0 channel") )
-						]
-				]
-			]
-			
-			+SVerticalBox::Slot()
-				.AutoHeight()
-				.Padding( 0.0f, 8.0f, 16.0f, 0.0f)
-			[
-				SNew(SHorizontalBox)
-				+SHorizontalBox::Slot()
-					.VAlign(VAlign_Center)
-					.FillWidth(3.0f)
-				[
-					SNew(STextBlock)
-						.Text( LOCTEXT("UVChannelSaveSelection", "UV channel to save results to:") )
-				]
-
-				+SHorizontalBox::Slot()
-					.HAlign(HAlign_Center)
-					.FillWidth(1.0f)
-				[
-					SAssignNew(UVChannelCombo, STextComboBox)
-						.OptionsSource(&UVChannels)
-						.InitiallySelectedItem(UVChannels.Num() > 1? UVChannels[1] : UVChannels[0])
-				]
-				
-			]
-
-			+SVerticalBox::Slot()
-				.AutoHeight()
-				.Padding( 0.0f, 8.0f, 0.0f, 0.0f)
-			[
-				SNew(STextBlock)
-					.Text( LOCTEXT("SelectGenerationMethod", "Select Generation Method") )
-					.IsEnabled(this, &SGenerateUniqueUVs::IsCreateNew)
-			]
-
-			+SVerticalBox::Slot()
-				.AutoHeight()
-				.Padding( 8.0f, 4.0f, 16.0f, 0.0f)
-			[
-				SNew(SHorizontalBox)
-				+SHorizontalBox::Slot()
-					.VAlign(VAlign_Center)
-					.FillWidth(3.0f)
-				[
-					SNew(SCheckBox)	
-						.Style(FEditorStyle::Get(), "RadioButton")
-						.IsChecked( this, &SGenerateUniqueUVs::IsLimitModeChecked, ELimitModeChoice::Stretching )
-						.OnCheckStateChanged( this, &SGenerateUniqueUVs::OnLimitModeChanged, ELimitModeChoice::Stretching )
-						.IsEnabled(this, &SGenerateUniqueUVs::IsCreateNew)
-					[
-						SNew(STextBlock)
-							.Text( LOCTEXT("LimitMaxStretching", "Limit maximum stretching") )
-					]
-				]
-
-				+SHorizontalBox::Slot()
-					.VAlign(VAlign_Center)
-					.FillWidth(1.0f)
-				[
-					SNew(SSpinBox<float>)
-						.MinValue(0.0f)
-						.MaxValue(1.0f)
-						.Value(0.5f)
-						.OnValueCommitted( this, &SGenerateUniqueUVs::OnMaxStretchingCommitted )
-						.OnValueChanged( this, &SGenerateUniqueUVs::OnMaxStretchingChanged )
-						.IsEnabled(this, &SGenerateUniqueUVs::IsStretchingLimit)
-				]
-			]
-
-			+SVerticalBox::Slot()
-				.AutoHeight()
-				.Padding( 8.0f, 8.0f, 16.0f, 0.0f)
-			[
-				SNew(SHorizontalBox)
-				+SHorizontalBox::Slot()
-					.VAlign(VAlign_Center)
-					.FillWidth(3.0f)
-				[
-					SNew(SCheckBox)	
-						.Style(FEditorStyle::Get(), "RadioButton")
-						.IsChecked( this, &SGenerateUniqueUVs::IsLimitModeChecked, ELimitModeChoice::Charts )
-						.OnCheckStateChanged( this, &SGenerateUniqueUVs::OnLimitModeChanged, ELimitModeChoice::Charts )
-						.IsEnabled(this, &SGenerateUniqueUVs::IsCreateNew)
-					[
-						SNew(STextBlock)
-							.Text( LOCTEXT("LimitMaxCharts", "Limit maximum number of charts") )
-					]
-				]
-
-				+SHorizontalBox::Slot()
-					.VAlign(VAlign_Center)
-					.FillWidth(1.0f)
-				[
-					SNew(SSpinBox<float>)
-						.MinValue(0.0f)
-						.MaxValue(100000.0f)
-						.Value( this, &SGenerateUniqueUVs::GetMaxCharts )
-						.OnValueCommitted( this, &SGenerateUniqueUVs::OnMaxChartsCommitted )
-						.OnValueChanged( this, &SGenerateUniqueUVs::OnMaxChartsChanged )
-						.IsEnabled(this, &SGenerateUniqueUVs::IsChartLimit)
-				]
-			]
-
-			+SVerticalBox::Slot()
-				.AutoHeight()
-				.Padding( 8.0f, 8.0f, 16.0f, 0.0f)
-			[
-				SNew(SHorizontalBox)
-				+SHorizontalBox::Slot()
-					.VAlign(VAlign_Center)
-					.FillWidth(3.0f)
-				[
-					SNew(STextBlock)
-						.Text( LOCTEXT("LimitChartSpacing", "Limit spacing between charts:") )
-						.IsEnabled(this, &SGenerateUniqueUVs::IsCreateNew)
-				]
-
-				+SHorizontalBox::Slot()
-					.VAlign(VAlign_Center)
-					.FillWidth(1.0f)
-				[
-					SNew(SSpinBox<float>)
-						.MinValue(0.0f)
-						.MaxValue(100.0f)
-						.Value(1.0f)
-						.OnValueCommitted( this, &SGenerateUniqueUVs::OnMinSpacingCommitted )
-						.OnValueChanged( this, &SGenerateUniqueUVs::OnMinSpacingChanged )
-						.IsEnabled(this, &SGenerateUniqueUVs::IsCreateNew)
-				]
-			]
-
-			+SVerticalBox::Slot()
-				.AutoHeight()
-				.HAlign(HAlign_Center)
-				.Padding( 0.0f, 8.0f, 0.0f, 0.0f)
-			[
-				SNew(SButton)
-					.Text( LOCTEXT("Apply", "Apply") )
-					.OnClicked( this, &SGenerateUniqueUVs::OnApply )
-			]
-		];
-
-	CurrentCreationModeChoice = CreateNew;
-	CurrentLimitModeChoice = Stretching;
-}
-END_SLATE_FUNCTION_BUILD_OPTIMIZATION
-
-SGenerateUniqueUVs::~SGenerateUniqueUVs()
-{
-
-}
-
-ESlateCheckBoxState::Type SGenerateUniqueUVs::IsCreationModeChecked( ECreationModeChoice ButtonId ) const
-{
-	return (CurrentCreationModeChoice == ButtonId)
-		? ESlateCheckBoxState::Checked
-		: ESlateCheckBoxState::Unchecked;
-}
-
-void SGenerateUniqueUVs::OnCreationModeChanged( ESlateCheckBoxState::Type NewRadioState, ECreationModeChoice RadioThatChanged )
-{
-	if (NewRadioState == ESlateCheckBoxState::Checked)
-	{
-		CurrentCreationModeChoice = RadioThatChanged;
-	}
-}
-
-ESlateCheckBoxState::Type SGenerateUniqueUVs::IsLimitModeChecked( enum ELimitModeChoice ButtonId ) const
-{
-	return (CurrentLimitModeChoice == ButtonId)
-	? ESlateCheckBoxState::Checked
-	: ESlateCheckBoxState::Unchecked;
-}
-
-void SGenerateUniqueUVs::OnLimitModeChanged( ESlateCheckBoxState::Type NewRadioState, enum ELimitModeChoice RadioThatChanged )
-{
-	if (NewRadioState == ESlateCheckBoxState::Checked)
-	{
-		CurrentLimitModeChoice = RadioThatChanged;
-	}
-}
-
-bool SGenerateUniqueUVs::IsCreateNew() const
-{
-	return CurrentCreationModeChoice == CreateNew;
-}
-
-bool SGenerateUniqueUVs::IsChartLimit() const
-{
-	return (CurrentLimitModeChoice == Charts) && IsCreateNew();
-}
-
-bool SGenerateUniqueUVs::IsStretchingLimit() const
-{
-	return (CurrentLimitModeChoice == Stretching) && IsCreateNew();
-}
-
-FReply SGenerateUniqueUVs::OnApply()
-{
-	int32 ChosenLODIndex = 0;
-	int32 ChosenUVChannel = UVChannels.Find(UVChannelCombo->GetSelectedItem());
-
-	bool bOnlyLayoutUVs = !IsCreateNew();
-	UStaticMesh* StaticMesh = StaticMeshEditorPtr.Pin()->GetStaticMesh();
-
-	if (StaticMesh->SourceModels.IsValidIndex(ChosenLODIndex)
-		&& !StaticMesh->SourceModels[ChosenLODIndex].RawMeshBulkData->IsEmpty())
-	{
-		bool bStatus = false;
-		uint32 Uint32MaxCharts = (uint32)MaxCharts;
-		FText Error;
-		GWarn->BeginSlowTask( NSLOCTEXT("UnrealEd", "GenerateUVsProgressText", "Generating unique UVs..."), true );
-		{
-			// Detach all instances of the static mesh while generating the UVs, then reattach them.
-			FStaticMeshComponentRecreateRenderStateContext RecreateRenderStateContext(StaticMeshEditorPtr.Pin()->GetStaticMesh());
-
-			// If the user has selected any edges in the static mesh editor, we'll create an array of chart UV
-			// seam edge indices to pass along to the GenerateUVs function.
-			TArray< int32 >* FalseEdgeIndicesPtr = NULL;
-			TArray< int32 > FalseEdgeIndices;
-			TSet< int32 >& SelectedEdgeIndices = StaticMeshEditorPtr.Pin()->GetSelectedEdges();
-
-			if( SelectedEdgeIndices.Num() > 0 &&
-				ChosenLODIndex == 0 )		// @todo: Support other LODs than LOD 0 (edge selection in SME needed)
-			{
-				for( TSet< int32 >::TIterator SelectionIt( SelectedEdgeIndices );
-					SelectionIt;
-					++SelectionIt )
-				{
-					const uint32 EdgeIndex = *SelectionIt;
-
-					FalseEdgeIndices.Add( EdgeIndex );
-				}
-
-				FalseEdgeIndicesPtr = &FalseEdgeIndices;
-			}
-
-			FRawMesh RawMesh;
-			StaticMesh->SourceModels[ChosenLODIndex].RawMeshBulkData->LoadRawMesh(RawMesh);
-
-			//call the utility helper with the user supplied parameters
-			{
-				IMeshUtilities& MeshUtils = FModuleManager::Get().LoadModuleChecked<IMeshUtilities>("MeshUtilities");
-	
-				// Don't display "border spacing" option as there's usually no reason to add border padding
-				float BorderSpacingPercent = 0;
-
-				bool bUseMaxStretch = IsStretchingLimit();
-
-				if(bOnlyLayoutUVs)
-				{
-					uint32 LightmapResolution = StaticMesh->LightMapResolution ? StaticMesh->LightMapResolution : 256;
-					bStatus = MeshUtils.LayoutUVs(RawMesh, LightmapResolution, ChosenUVChannel, Error);
-				}		
-				else 
-				{
-					bStatus = MeshUtils.GenerateUVs(RawMesh, ChosenUVChannel, MinSpacingBetweenCharts, BorderSpacingPercent, bUseMaxStretch, FalseEdgeIndicesPtr, Uint32MaxCharts, MaxStretching, Error);
-				}
-			}
-
-			if (bStatus)
-			{
-				StaticMesh->SourceModels[ChosenLODIndex].RawMeshBulkData->SaveRawMesh(RawMesh);
-				StaticMesh->Build();
-			}
-		}
-		GWarn->EndSlowTask();
-
-		FText StatusMessage;
-		if(bStatus)
-		{
-			FNumberFormattingOptions NumberOptions;
-			NumberOptions.MinimumFractionalDigits = 2;
-			NumberOptions.MaximumFractionalDigits = 2;
-
-			FFormatNamedArguments Args;
-			Args.Add( TEXT("MaxCharts"), Uint32MaxCharts );
-			Args.Add( TEXT("MaxStretching"), FText::AsPercent( MaxStretching, &NumberOptions ) );
-			StatusMessage = FText::Format( NSLOCTEXT("GenerateUVsWindow", "GenerateUniqueUVs_UVGenerationSuccessful", "Finished generating UVs; Charts: {MaxCharts}; Largest stretch: {MaxStretching}"), Args );
-		}
-		else if (!Error.IsEmpty())
-		{
-			StatusMessage = Error;
-		}
-		else
-		{
-			StatusMessage = NSLOCTEXT("GenerateUVsWindow", "GenerateUniqueUVs_UVGenerationFailed", "Mesh UV generation failed.");
-		}
-
-		FNotificationInfo NotificationInfo( StatusMessage );
-		NotificationInfo.ExpireDuration = 3.0f;
-
-		TSharedPtr<SNotificationItem> NotificationItem = FSlateNotificationManager::Get().AddNotification(NotificationInfo);
-		if ( NotificationItem.IsValid() )
-		{
-			NotificationItem->SetCompletionState(bStatus ? SNotificationItem::CS_Success : SNotificationItem::CS_Fail);
-		}
-	}
-
-	StaticMeshEditorPtr.Pin()->RefreshTool();
-	RefreshTool();
-
-	return FReply::Handled();
-}
-
-void SGenerateUniqueUVs::OnMaxStretchingChanged(float InValue)
-{
-	MaxStretching = InValue;
-}
-
-void SGenerateUniqueUVs::OnMaxStretchingCommitted(float InValue, ETextCommit::Type CommitInfo)
-{
-	OnMaxStretchingChanged(InValue);
-}
-
-void SGenerateUniqueUVs::OnMaxChartsChanged(float InValue)
-{
-	MaxCharts = float(int32(InValue));
-}
-
-void SGenerateUniqueUVs::OnMaxChartsCommitted(float InValue, ETextCommit::Type CommitInfo)
-{
-	OnMaxChartsChanged(InValue);
-}
-
-float SGenerateUniqueUVs::GetMaxCharts() const
-{
-	return MaxCharts;
-}
-
-void SGenerateUniqueUVs::OnMinSpacingChanged(float InValue)
-{
-	MinSpacingBetweenCharts = InValue;
-}
-
-void SGenerateUniqueUVs::OnMinSpacingCommitted(float InValue, ETextCommit::Type CommitInfo)
-{
-	OnMinSpacingChanged(InValue);
-}
-
-void SGenerateUniqueUVs::RefreshUVChannelList()
-{
-	UStaticMesh* StaticMesh = StaticMeshEditorPtr.Pin()->GetStaticMesh();
-
-	// Fill out the UV channels combo.
-	UVChannels.Empty();
-	for(int32 UVChannelID = 0; UVChannelID < FMath::Min(StaticMeshEditorPtr.Pin()->GetNumUVChannels() + 1, (int32)MAX_STATIC_TEXCOORDS); ++UVChannelID)
-	{
-		UVChannels.Add( MakeShareable( new FString( FText::Format( LOCTEXT("UVChannel_ID", "UV Channel {0}"), FText::AsNumber( UVChannelID ) ).ToString() ) ) );
-	}
-
-	if(UVChannelCombo.IsValid())
-	{
-		if(UVChannels.Num() > 1)
-		{
-			UVChannelCombo->SetSelectedItem(UVChannels[1]);
-		}
-		else
-		{
-			UVChannelCombo->SetSelectedItem(UVChannels[0]);
-		}
-	}
-	
-}
-
-void SGenerateUniqueUVs::RefreshTool()
-{
-	RefreshUVChannelList();
-}
-
 #undef LOCTEXT_NAMESPACE
diff --git a/Engine/Source/Editor/StaticMeshEditor/Private/StaticMeshEditorTools.h b/Engine/Source/Editor/StaticMeshEditor/Private/StaticMeshEditorTools.h
--- a/Engine/Source/Editor/StaticMeshEditor/Private/StaticMeshEditorTools.h
+++ b/Engine/Source/Editor/StaticMeshEditor/Private/StaticMeshEditorTools.h
@@ -331,101 +331,4 @@ private:
 	bool bBuildSettingsExpanded[MAX_STATIC_MESH_LODS];
 	bool bReductionSettingsExpanded[MAX_STATIC_MESH_LODS];
 	bool bSectionSettingsExpanded[MAX_STATIC_MESH_LODS];
-};
-
-/** 
- * Window for generating unique UVs for the static mesh.
- */
-class SGenerateUniqueUVs : public SCompoundWidget
-{
-public:
-	SLATE_BEGIN_ARGS( SGenerateUniqueUVs ) :
-	  _StaticMeshEditorPtr()
-	  {
-	  }
-	  /** The Static Mesh Editor this tool is associated with. */
-	  SLATE_ARGUMENT( TWeakPtr< IStaticMeshEditor >, StaticMeshEditorPtr )
-
-	SLATE_END_ARGS()
-
-	void Construct(const FArguments& InArgs);
-
-	virtual ~SGenerateUniqueUVs();
-
-	/** Refreshes the tool's LOD and UV channel options. */
-	void RefreshTool();
-
-private:
-	/** Decides if the radio button should be checked. */
-	ESlateCheckBoxState::Type IsCreationModeChecked( enum ECreationModeChoice ButtonId ) const;
-
-	/** Handles the Creation mode being changed and manages the radio buttons. */
-	void OnCreationModeChanged( ESlateCheckBoxState::Type NewRadioState, enum ECreationModeChoice RadioThatChanged );
-
-	/** Decides if the radio button should be checked. */
-	ESlateCheckBoxState::Type IsLimitModeChecked( enum ELimitModeChoice ButtonId ) const;
-
-	/** Handles the Limit mode being changed and manages the radio buttons. */
-	void OnLimitModeChanged( ESlateCheckBoxState::Type NewRadioState, enum ELimitModeChoice RadioThatChanged );
-
-	/** Checks the state of the Creation Mode, TRUE if Create New is selected. */
-	bool IsCreateNew() const;
-
-	/** Checks the state of the limit mode, TRUE if charts is selected. */
-	bool IsChartLimit() const;
-
-	/** Checks the state of the limit mode, TRUE if stretching is selected. */
-	bool IsStretchingLimit() const;
-
-	/** Callback for the Apply button. */
-	FReply OnApply();
-
-	/** Callback for spinbox to set the MaxStretching variable. */
-	void OnMaxStretchingChanged(float InValue);
-
-	/** Callback for spinbox to set the MaxStretching variable. */
-	void OnMaxStretchingCommitted(float InValue, ETextCommit::Type CommitInfo);
-
-	/** Callback for spinbox to set the MaxCharts variable. Reduces the value to an int32 */
-	void OnMaxChartsChanged(float InValue);
-
-	/** Callback for spinbox to set the MaxCharts variable. Reduces the value to an int32 */
-	void OnMaxChartsCommitted(float InValue, ETextCommit::Type CommitInfo);
-	
-	/** Retrieves the current value for MaxCharts for a spinbox, important because the value is limited to an int32. */
-	float GetMaxCharts() const;
-
-	/** Callback for spinbox to set the MinSpacing variable. */
-	void OnMinSpacingChanged(float InValue);
-
-	/** Callback for spinbox to set the MinSpacing variable. */
-	void OnMinSpacingCommitted(float InValue, ETextCommit::Type CommitInfo);
-
-	/** Refreshes the UV Channel list. */
-	void RefreshUVChannelList();
-
-private:
-	/** The Static Mesh Editor this tool is associated with. */
-	TWeakPtr< IStaticMeshEditor > StaticMeshEditorPtr;
-
-	/** Drop down menu for selecting the UV Channel to save to. */
-	TSharedPtr<class STextComboBox> UVChannelCombo;
-
-	/** Options list for the UV Channel options. */
-	TArray< TSharedPtr< FString > > UVChannels;
-
-	/** The current Creation mode selected by the radio buttons. */
-	ECreationModeChoice CurrentCreationModeChoice;
-
-	/** The current Limit mode selected by the radio buttons. */
-	ELimitModeChoice CurrentLimitModeChoice;
-
-	/** Modified by a spinbox, user setting for max stretching. */
-	float MaxStretching;
-
-	/** Modified by a spinbox, user setting for max charts, clamped to be an int32. */
-	float MaxCharts;
-
-	/** Modified by a spinbox, user setting for the min spacing between charts. */
-	float MinSpacingBetweenCharts;
-};
+};
\ No newline at end of file
diff --git a/Engine/Source/Editor/StaticMeshEditor/Public/StaticMeshEditorActions.h b/Engine/Source/Editor/StaticMeshEditor/Public/StaticMeshEditorActions.h
--- a/Engine/Source/Editor/StaticMeshEditor/Public/StaticMeshEditorActions.h
+++ b/Engine/Source/Editor/StaticMeshEditor/Public/StaticMeshEditorActions.h
@@ -59,8 +59,6 @@ public:
 	// Mesh Menu Commands
 	TSharedPtr< FUICommandInfo > FindSource;
 
-	TSharedPtr< FUICommandInfo > GenerateUniqueUVs;
-
 	TSharedPtr< FUICommandInfo > ChangeMesh;
 
 	/**
