Improve auto collision generation.
- if any changes to collision or custom collision don't generate
- advanced option to let you disable collision generation

diff --git a/Engine/Source/Editor/StaticMeshEditor/Private/StaticMeshEditor.cpp b/Engine/Source/Editor/StaticMeshEditor/Private/StaticMeshEditor.cpp
--- a/Engine/Source/Editor/StaticMeshEditor/Private/StaticMeshEditor.cpp
+++ b/Engine/Source/Editor/StaticMeshEditor/Private/StaticMeshEditor.cpp
@@ -687,6 +687,8 @@ void FStaticMeshEditor::DuplicateSelectedPrims(const FVector* InOffset)
 
 		// Update views/property windows
 		Viewport->RefreshViewport();
+
+		StaticMesh->bCustomizedCollision = true;	//mark the static mesh for collision customization
 	}
 }
 
@@ -768,6 +770,8 @@ void FStaticMeshEditor::ScaleSelectedPrims(const FVector& InScale)
 			AggGeom->ConvexElems[PrimData.PrimIndex].ScaleElem(ModifiedScale, MinPrimSize);
 			break;
 		}
+
+		StaticMesh->bCustomizedCollision = true;	//mark the static mesh for collision customization
 	}
 }
 
@@ -874,6 +878,8 @@ void FStaticMeshEditor::SetPrimTransform(const FPrimData& InPrimData, const FTra
 		AggGeom->ConvexElems[InPrimData.PrimIndex].SetTransform(InPrimTransform);
 		break;
 	}
+
+	StaticMesh->bCustomizedCollision = true;	//mark the static mesh for collision customization
 }
 
 void FStaticMeshEditor::RefreshTool()
@@ -1123,6 +1129,8 @@ void FStaticMeshEditor::OnRemoveCollision(void)
 
 			// Update views/property windows
 			Viewport->RefreshViewport();
+
+			StaticMesh->bCustomizedCollision = true;	//mark the static mesh for collision customization
 		}
 	}
 }
@@ -1240,6 +1248,8 @@ void FStaticMeshEditor::OnConvertBoxToConvexCollision()
 
 				// Update views/property windows
 				Viewport->RefreshViewport();
+
+				StaticMesh->bCustomizedCollision = true;	//mark the static mesh for collision customization
 			}
 		}
 	}
@@ -1295,6 +1305,8 @@ void FStaticMeshEditor::OnCopyCollisionFromSelectedStaticMesh()
 
 			// Update views/property windows
 			Viewport->RefreshViewport();
+
+			StaticMesh->bCustomizedCollision = true;	//mark the static mesh for collision customization
 		}
 	}
 	else
@@ -1437,6 +1449,8 @@ void FStaticMeshEditor::DoDecomp(int32 InMaxHullCount, int32 InMaxHullVerts)
 
 		// Update screen.
 		Viewport->RefreshViewport();
+
+		StaticMesh->bCustomizedCollision = true;	//mark the static mesh for collision customization
 	}
 }
 
@@ -1590,6 +1604,8 @@ void FStaticMeshEditor::DeleteSelectedPrims()
 
 			// Update views/property windows
 			Viewport->RefreshViewport();
+
+			StaticMesh->bCustomizedCollision = true;	//mark the static mesh for collision customization
 		}
 	}
 }
diff --git a/Engine/Source/Editor/UnrealEd/Classes/Factories/FbxStaticMeshImportData.h b/Engine/Source/Editor/UnrealEd/Classes/Factories/FbxStaticMeshImportData.h
--- a/Engine/Source/Editor/UnrealEd/Classes/Factories/FbxStaticMeshImportData.h
+++ b/Engine/Source/Editor/UnrealEd/Classes/Factories/FbxStaticMeshImportData.h
@@ -31,6 +31,10 @@ class UFbxStaticMeshImportData : public UFbxMeshImportData
 	UPROPERTY(EditAnywhere, config, Category=ImportSettings)
 	uint32 bOneConvexHullPerUCX:1;
 
+	/** If checked, collision will automatically be generated (ignored if custom collision is imported or used). */
+	UPROPERTY(EditAnywhere, config, Category = ImportSettings)
+	uint32 bAutoGenerateCollision : 1;
+
 	/** Gets or creates fbx import data for the specified static mesh */
 	static UFbxStaticMeshImportData* GetImportDataForStaticMesh(UStaticMesh* StaticMesh, UFbxStaticMeshImportData* TemplateForCreation);
 };
diff --git a/Engine/Source/Editor/UnrealEd/Private/Fbx/FbxMainImport.cpp b/Engine/Source/Editor/UnrealEd/Private/Fbx/FbxMainImport.cpp
--- a/Engine/Source/Editor/UnrealEd/Private/Fbx/FbxMainImport.cpp
+++ b/Engine/Source/Editor/UnrealEd/Private/Fbx/FbxMainImport.cpp
@@ -197,6 +197,7 @@ void ApplyImportUIToImportOptions(UFbxImportUI* ImportUI, FBXImportOptions& InOu
 	InOutImportOptions.bRemoveDegenerates = ImportUI->StaticMeshImportData->bRemoveDegenerates;
 	InOutImportOptions.bGenerateLightmapUVs = ImportUI->StaticMeshImportData->bGenerateLightmapUVs;
 	InOutImportOptions.bOneConvexHullPerUCX = ImportUI->StaticMeshImportData->bOneConvexHullPerUCX;
+	InOutImportOptions.bAutoGenerateCollision = ImportUI->StaticMeshImportData->bAutoGenerateCollision;
 	InOutImportOptions.StaticMeshLODGroup = ImportUI->StaticMeshImportData->StaticMeshLODGroup;
 	InOutImportOptions.bImportMeshesInBoneHierarchy = ImportUI->SkeletalMeshImportData->bImportMeshesInBoneHierarchy;
 	InOutImportOptions.bCreatePhysicsAsset = ImportUI->bCreatePhysicsAsset;
diff --git a/Engine/Source/Editor/UnrealEd/Private/Fbx/FbxOptionWindow.cpp b/Engine/Source/Editor/UnrealEd/Private/Fbx/FbxOptionWindow.cpp
--- a/Engine/Source/Editor/UnrealEd/Private/Fbx/FbxOptionWindow.cpp
+++ b/Engine/Source/Editor/UnrealEd/Private/Fbx/FbxOptionWindow.cpp
@@ -745,6 +745,18 @@ TSharedRef<SWidget> SFbxOptionWindow::ConstructStaticMeshAdvanced()
 		]
 	];
 
+	NewBox->AddSlot().AutoHeight().Padding(2)
+		[
+			SNew(SCheckBox)
+			.IsChecked(ImportUI->StaticMeshImportData->bAutoGenerateCollision ? ESlateCheckBoxState::Checked : ESlateCheckBoxState::Unchecked)
+			.OnCheckStateChanged(this, &SFbxOptionWindow::SetStaticMesh_AutoGenerateCollision)
+			.IsEnabled(!bIsObjFormat)
+			[
+				SNew(STextBlock)
+				.Text(LOCTEXT("FbxOptionWindow_AutoGenerateCollision", "Auto Generate Collision"))
+			]
+		];
+
 
 	return NewBox;
 }
@@ -1268,6 +1280,11 @@ void SFbxOptionWindow::SetStaticMesh_OneConvexHullPerUCX(ESlateCheckBoxState::Ty
 	ImportUI->StaticMeshImportData->bOneConvexHullPerUCX = (NewType == ESlateCheckBoxState::Checked)? 1:0;
 }
 
+void SFbxOptionWindow::SetStaticMesh_AutoGenerateCollision(ESlateCheckBoxState::Type NewType)
+{
+	ImportUI->StaticMeshImportData->bAutoGenerateCollision = (NewType == ESlateCheckBoxState::Checked) ? 1 : 0;
+}
+
 void SFbxOptionWindow::SetMaterial_ImportMaterials(ESlateCheckBoxState::Type NewType)
 {
 	ImportUI->bImportMaterials = (NewType == ESlateCheckBoxState::Checked)? 1:0;
diff --git a/Engine/Source/Editor/UnrealEd/Private/Fbx/FbxStaticMeshImport.cpp b/Engine/Source/Editor/UnrealEd/Private/Fbx/FbxStaticMeshImport.cpp
--- a/Engine/Source/Editor/UnrealEd/Private/Fbx/FbxStaticMeshImport.cpp
+++ b/Engine/Source/Editor/UnrealEd/Private/Fbx/FbxStaticMeshImport.cpp
@@ -394,8 +394,20 @@ bool UnFbx::FFbxImporter::BuildStaticMeshFromGeometry(FbxMesh* Mesh, UStaticMesh
 	//
 	// build collision
 	//
-	bool bEnableCollision = ImportCollisionModels(StaticMesh, GetNodeNameWithoutNamespace( Node ) )
-		|| (GBuildStaticMeshCollision && LODIndex == 0 && ImportOptions->bRemoveDegenerates);
+	bool bImportedCollision = ImportCollisionModels(StaticMesh, GetNodeNameWithoutNamespace(Node));
+
+	if (false && !bImportedCollision && StaticMesh)	//if didn't import collision automatically generate one
+	{
+		StaticMesh->CreateBodySetup();
+
+		const int32 NumDirs = 18;
+		TArray<FVector> Dirs;
+		Dirs.AddUninitialized(NumDirs);
+		for (int32 DirIdx = 0; DirIdx < NumDirs; ++DirIdx) { Dirs[DirIdx] = KDopDir18[DirIdx]; }
+		GenerateKDopAsSimpleCollision(StaticMesh, Dirs);
+	}
+
+	bool bEnableCollision = bImportedCollision || (GBuildStaticMeshCollision && LODIndex == 0 && ImportOptions->bRemoveDegenerates);
 	for(int32 SectionIndex=MaterialIndexOffset; SectionIndex<MaterialIndexOffset+MaterialCount; SectionIndex++)
 	{
 		FMeshSectionInfo Info = StaticMesh->SectionInfoMap.Get(LODIndex, SectionIndex);
@@ -1200,12 +1212,11 @@ UStaticMesh* UnFbx::FFbxImporter::ImportStaticMeshAsSingle(UObject* InParent, TA
 		StaticMesh = NULL;
 	}
 
-
-	if (StaticMesh)
+	if (StaticMesh && StaticMesh->bCustomizedCollision == false && ImportOptions->bAutoGenerateCollision)
 	{
 
 		FKAggregateGeom & AggGeom = StaticMesh->BodySetup->AggGeom;
-		AggGeom.ConvexElems.Empty(1);	//when reimport happens we remove any existing convex mesh
+		AggGeom.ConvexElems.Empty(1);	//if no custom collision is setup we just regenerate collision when reimport
 
 		const int32 NumDirs = 18;
 		TArray<FVector> Dirs;
@@ -1389,6 +1400,8 @@ bool UnFbx::FFbxImporter::ImportCollisionModels(UStaticMesh* StaticMesh, const F
 
 	TSharedPtr< FbxArray<FbxNode*> > Models = Record->GetValue();
 
+	StaticMesh->bCustomizedCollision = true;
+
 	StaticMesh->CreateBodySetup();
 
 	TArray<FVector>	CollisionVertices;
diff --git a/Engine/Source/Editor/UnrealEd/Private/Fbx/FbxStaticMeshImportData.cpp b/Engine/Source/Editor/UnrealEd/Private/Fbx/FbxStaticMeshImportData.cpp
--- a/Engine/Source/Editor/UnrealEd/Private/Fbx/FbxStaticMeshImportData.cpp
+++ b/Engine/Source/Editor/UnrealEd/Private/Fbx/FbxStaticMeshImportData.cpp
@@ -9,6 +9,7 @@ UFbxStaticMeshImportData::UFbxStaticMeshImportData(const class FPostConstructIni
 	bRemoveDegenerates = true;
 	bGenerateLightmapUVs = true;
 	bOneConvexHullPerUCX = true;
+	bAutoGenerateCollision = true;
 }
 
 UFbxStaticMeshImportData* UFbxStaticMeshImportData::GetImportDataForStaticMesh(UStaticMesh* StaticMesh, UFbxStaticMeshImportData* TemplateForCreation)
diff --git a/Engine/Source/Editor/UnrealEd/Private/FbxImporter.h b/Engine/Source/Editor/UnrealEd/Private/FbxImporter.h
--- a/Engine/Source/Editor/UnrealEd/Private/FbxImporter.h
+++ b/Engine/Source/Editor/UnrealEd/Private/FbxImporter.h
@@ -115,6 +115,7 @@ struct FBXImportOptions
 	bool bRemoveDegenerates;
 	bool bGenerateLightmapUVs;
 	bool bOneConvexHullPerUCX;
+	bool bAutoGenerateCollision;
 	FName StaticMeshLODGroup;
 	// Skeletal Mesh options
 	bool bImportMorph;
diff --git a/Engine/Source/Editor/UnrealEd/Private/FbxOptionWindow.h b/Engine/Source/Editor/UnrealEd/Private/FbxOptionWindow.h
--- a/Engine/Source/Editor/UnrealEd/Private/FbxOptionWindow.h
+++ b/Engine/Source/Editor/UnrealEd/Private/FbxOptionWindow.h
@@ -118,6 +118,7 @@ private:
 	void SetStaticMesh_RemoveDegenerates(ESlateCheckBoxState::Type NewType);
 	void SetStaticMesh_GenerateLightmapUVs(ESlateCheckBoxState::Type NewType);
 	void SetStaticMesh_OneConvexHullPerUCX(ESlateCheckBoxState::Type NewType);
+	void SetStaticMesh_AutoGenerateCollision(ESlateCheckBoxState::Type NewType);
 	void SetMaterial_ImportMaterials(ESlateCheckBoxState::Type NewType);
 	void SetMaterial_ImportTextures(ESlateCheckBoxState::Type NewType);
 	void SetMaterial_InvertNormalMaps(ESlateCheckBoxState::Type NewType);
diff --git a/Engine/Source/Editor/UnrealEd/Private/GeomFitUtils.cpp b/Engine/Source/Editor/UnrealEd/Private/GeomFitUtils.cpp
--- a/Engine/Source/Editor/UnrealEd/Private/GeomFitUtils.cpp
+++ b/Engine/Source/Editor/UnrealEd/Private/GeomFitUtils.cpp
@@ -229,6 +229,8 @@ int32 GenerateBoxAsSimpleCollision(UStaticMesh* StaticMesh)
 	// Mark staticmesh as dirty, to help make sure it gets saved.
 	StaticMesh->MarkPackageDirty();
 
+	StaticMesh->bCustomizedCollision = true;	//mark the static mesh for collision customization
+
 	return bs->AggGeom.BoxElems.Num() - 1;
 }
 
@@ -409,6 +411,7 @@ int32 GenerateSphereAsSimpleCollision(UStaticMesh* StaticMesh)
 	// Mark staticmesh as dirty, to help make sure it gets saved.
 	StaticMesh->MarkPackageDirty();
 
+	StaticMesh->bCustomizedCollision = true;	//mark the static mesh for collision customization
 	return bs->AggGeom.SphereElems.Num() - 1;
 }
 
@@ -551,6 +554,8 @@ int32 GenerateSphylAsSimpleCollision(UStaticMesh* StaticMesh)
 	// Mark staticmesh as dirty, to help make sure it gets saved.
 	StaticMesh->MarkPackageDirty();
 
+	StaticMesh->bCustomizedCollision = true;	//mark the static mesh for collision customization
+
 	return bs->AggGeom.SphylElems.Num() - 1;
 }
 
diff --git a/Engine/Source/Editor/UnrealEd/Private/StaticMeshEdit.cpp b/Engine/Source/Editor/UnrealEd/Private/StaticMeshEdit.cpp
--- a/Engine/Source/Editor/UnrealEd/Private/StaticMeshEdit.cpp
+++ b/Engine/Source/Editor/UnrealEd/Private/StaticMeshEdit.cpp
@@ -822,6 +822,7 @@ struct ExistingStaticMeshData
 	TArray<UStaticMeshSocket*>	ExistingSockets;
 
 	bool						ExistingUseMaximumStreamingTexelRatio;
+	bool						ExistingCustomizedCollision;
 
 	int32							ExistingLightMapResolution;
 	int32							ExistingLightMapCoordinateIndex;
@@ -859,6 +860,7 @@ ExistingStaticMeshData* SaveExistingStaticMeshData(UStaticMesh* ExistingMesh)
 		ExistingMeshDataPtr->ExistingSockets = ExistingMesh->Sockets;
 
 		ExistingMeshDataPtr->ExistingUseMaximumStreamingTexelRatio = ExistingMesh->bUseMaximumStreamingTexelRatio;
+		ExistingMeshDataPtr->ExistingCustomizedCollision = ExistingMesh->bCustomizedCollision;
 
 		ExistingMeshDataPtr->ExistingLightMapResolution = ExistingMesh->LightMapResolution;
 		ExistingMeshDataPtr->ExistingLightMapCoordinateIndex = ExistingMesh->LightMapCoordinateIndex;
@@ -906,6 +908,7 @@ void RestoreExistingMeshData(struct ExistingStaticMeshData* ExistingMeshDataPtr,
 		}
 
 		NewMesh->bUseMaximumStreamingTexelRatio = ExistingMeshDataPtr->ExistingUseMaximumStreamingTexelRatio;
+		NewMesh->bCustomizedCollision = ExistingMeshDataPtr->ExistingCustomizedCollision;
 
 		NewMesh->LightMapResolution = ExistingMeshDataPtr->ExistingLightMapResolution;
 		NewMesh->LightMapCoordinateIndex = ExistingMeshDataPtr->ExistingLightMapCoordinateIndex;
