Enable Static Mesh Collision view, by default

#nonttp

#branch UE4

diff --git a/Engine/Source/Editor/StaticMeshEditor/Private/StaticMeshEditorViewportClient.cpp b/Engine/Source/Editor/StaticMeshEditor/Private/StaticMeshEditorViewportClient.cpp
--- a/Engine/Source/Editor/StaticMeshEditor/Private/StaticMeshEditorViewportClient.cpp
+++ b/Engine/Source/Editor/StaticMeshEditor/Private/StaticMeshEditorViewportClient.cpp
@@ -67,7 +67,7 @@ FStaticMeshEditorViewportClient::FStaticMeshEditorViewportClient(TWeakPtr<IStati
 	OverrideNearClipPlane(1.0f);
 	bUsingOrbitCamera = true;
 
-	bShowCollision = false;
+	bShowCollision = true;
 	bShowSockets = true;
 	bDrawUVs = false;
 	bDrawNormals = false;
