Prevent new collision overlapping existing collision

#ttp 345297 - NUX: Improve workflow for SME collision editing

#branch UE4

#change Modified Collision options to say 'Add' to indicate you're adding primitives and not replacing existing.

#change When adding a new collision primitive, it now does a quick check to see if there's another primitive that's the same and if so, moves the new one slightly so they don't overlap - the hope is that this also makes it clearer that you're adding primitives.

Note: Also includes change from 2280280 which enables Show Collision if it's modified

diff --git a/Engine/Source/Editor/StaticMeshEditor/Private/StaticMeshEditor.cpp b/Engine/Source/Editor/StaticMeshEditor/Private/StaticMeshEditor.cpp
--- a/Engine/Source/Editor/StaticMeshEditor/Private/StaticMeshEditor.cpp
+++ b/Engine/Source/Editor/StaticMeshEditor/Private/StaticMeshEditor.cpp
@@ -588,6 +588,13 @@ bool FStaticMeshEditor::HasSelectedPrims() const
 void FStaticMeshEditor::AddSelectedPrim(const FPrimData& InPrimData)
 {
 	check(IsPrimValid(InPrimData));
+
+	// Enable collision, if not already
+	if( !Viewport->GetViewportClient().IsSetShowWireframeCollisionChecked() )
+	{
+		Viewport->GetViewportClient().SetShowWireframeCollision();
+	}
+
 	SelectedPrims.Add(InPrimData);	
 }
 
@@ -694,10 +701,6 @@ void FStaticMeshEditor::DuplicateSelectedPrims(const FVector* InOffset)
 
 void FStaticMeshEditor::TranslateSelectedPrims(const FVector& InDrag)
 {
-	check(StaticMesh->BodySetup);
-
-	FKAggregateGeom* AggGeom = &StaticMesh->BodySetup->AggGeom;
-
 	for (int32 PrimIdx = 0; PrimIdx < SelectedPrims.Num(); PrimIdx++)
 	{
 		const FPrimData& PrimData = SelectedPrims[PrimIdx];
@@ -714,10 +717,6 @@ void FStaticMeshEditor::TranslateSelectedPrims(const FVector& InDrag)
 
 void FStaticMeshEditor::RotateSelectedPrims(const FRotator& InRot)
 {
-	check(StaticMesh->BodySetup);
-
-	FKAggregateGeom* AggGeom = &StaticMesh->BodySetup->AggGeom;
-
 	const FQuat DeltaQ = InRot.Quaternion();
 
 	for (int32 PrimIdx = 0; PrimIdx < SelectedPrims.Num(); PrimIdx++)
@@ -882,6 +881,101 @@ void FStaticMeshEditor::SetPrimTransform(const FPrimData& InPrimData, const FTra
 	StaticMesh->bCustomizedCollision = true;	//mark the static mesh for collision customization
 }
 
+bool FStaticMeshEditor::OverlapsExistingPrim(const FPrimData& InPrimData) const
+{
+	check(StaticMesh->BodySetup);
+
+	const FKAggregateGeom* AggGeom = &StaticMesh->BodySetup->AggGeom;
+
+	// Assume that if the transform of the prim is the same, then it overlaps (FKConvexElem doesn't have an operator==, and no shape takes tolerances into account)
+	check(IsPrimValid(InPrimData));
+	switch (InPrimData.PrimType)
+	{
+	case KPT_Sphere:
+		{
+			const FKSphereElem InSphereElem = AggGeom->SphereElems[InPrimData.PrimIndex];
+			const FTransform InElemTM = InSphereElem.GetTransform();
+			for (int32 i = 0; i < AggGeom->SphereElems.Num(); ++i)
+			{
+				if( i == InPrimData.PrimIndex )
+				{
+					continue;
+				}
+
+				const FKSphereElem& SphereElem = AggGeom->SphereElems[i];
+				const FTransform ElemTM = SphereElem.GetTransform();
+				if( InElemTM.Equals(ElemTM) )
+				{
+					return true;
+				}
+			}
+		}
+		break;
+	case KPT_Box:
+		{
+			const FKBoxElem InBoxElem = AggGeom->BoxElems[InPrimData.PrimIndex];
+			const FTransform InElemTM = InBoxElem.GetTransform();
+			for (int32 i = 0; i < AggGeom->BoxElems.Num(); ++i)
+			{
+				if( i == InPrimData.PrimIndex )
+				{
+					continue;
+				}
+
+				const FKBoxElem& BoxElem = AggGeom->BoxElems[i];
+				const FTransform ElemTM = BoxElem.GetTransform();
+				if( InElemTM.Equals(ElemTM) )
+				{
+					return true;
+				}
+			}
+		}
+		break;
+	case KPT_Sphyl:
+		{
+			const FKSphylElem InSphylElem = AggGeom->SphylElems[InPrimData.PrimIndex];
+			const FTransform InElemTM = InSphylElem.GetTransform();
+			for (int32 i = 0; i < AggGeom->SphylElems.Num(); ++i)
+			{
+				if( i == InPrimData.PrimIndex )
+				{
+					continue;
+				}
+
+				const FKSphylElem& SphylElem = AggGeom->SphylElems[i];
+				const FTransform ElemTM = SphylElem.GetTransform();
+				if( InElemTM.Equals(ElemTM) )
+				{
+					return true;
+				}
+			}
+		}
+		break;
+	case KPT_Convex:
+		{
+			const FKConvexElem InConvexElem = AggGeom->ConvexElems[InPrimData.PrimIndex];
+			const FTransform InElemTM = InConvexElem.GetTransform();
+			for (int32 i = 0; i < AggGeom->ConvexElems.Num(); ++i)
+			{
+				if( i == InPrimData.PrimIndex )
+				{
+					continue;
+				}
+
+				const FKConvexElem& ConvexElem = AggGeom->ConvexElems[i];
+				const FTransform ElemTM = ConvexElem.GetTransform();
+				if( InElemTM.Equals(ElemTM) )
+				{
+					return true;
+				}
+			}
+		}
+		break;
+	}
+
+	return false;
+}
+
 void FStaticMeshEditor::RefreshTool()
 {
 	int32 NumLODs = StaticMesh->GetNumLODs();
@@ -1044,8 +1138,13 @@ void FStaticMeshEditor::GenerateKDop(const FVector* Directions, uint32 NumDirect
 		{
 			FEngineAnalytics::GetProvider().RecordEvent(TEXT("Editor.Usage.StaticMesh.Collision"), TEXT("Type"), TEXT("KDop Collision"));
 		}
+		const FPrimData PrimData = FPrimData(KPT_Convex, PrimIndex);
 		ClearSelectedPrims();
-		AddSelectedPrim(FPrimData(KPT_Convex, PrimIndex));
+		AddSelectedPrim(PrimData);
+		while( OverlapsExistingPrim(PrimData) )
+		{
+			TranslateSelectedPrims(OverlapNudge);
+		}
 	}
 
 	Viewport->RefreshViewport();
@@ -1062,8 +1161,13 @@ void FStaticMeshEditor::OnCollisionBox()
 		{
 			FEngineAnalytics::GetProvider().RecordEvent(TEXT("Editor.Usage.StaticMesh.Collision"), TEXT("Type"), TEXT("Box Collision"));
 		}
+		const FPrimData PrimData = FPrimData(KPT_Box, PrimIndex);
 		ClearSelectedPrims();
-		AddSelectedPrim(FPrimData(KPT_Box, PrimIndex));
+		AddSelectedPrim(PrimData);
+		while( OverlapsExistingPrim(PrimData) )
+		{
+			TranslateSelectedPrims(OverlapNudge);
+		}
 	}
 
 	Viewport->RefreshViewport();
@@ -1080,8 +1184,13 @@ void FStaticMeshEditor::OnCollisionSphere()
 		{
 			FEngineAnalytics::GetProvider().RecordEvent(TEXT("Editor.Usage.StaticMesh.Collision"), TEXT("Type"), TEXT("Sphere Collision"));
 		}
+		const FPrimData PrimData = FPrimData(KPT_Sphere, PrimIndex);
 		ClearSelectedPrims();
-		AddSelectedPrim(FPrimData(KPT_Sphere, PrimIndex));
+		AddSelectedPrim(PrimData);
+		while( OverlapsExistingPrim(PrimData) )
+		{
+			TranslateSelectedPrims(OverlapNudge);
+		}
 	}
 
 	Viewport->RefreshViewport();
@@ -1098,8 +1207,13 @@ void FStaticMeshEditor::OnCollisionSphyl()
 		{
 			FEngineAnalytics::GetProvider().RecordEvent(TEXT("Editor.Usage.StaticMesh.Collision"), TEXT("Type"), TEXT("Capsule Collision"));
 		}
+		const FPrimData PrimData = FPrimData(KPT_Sphyl, PrimIndex);
 		ClearSelectedPrims();
-		AddSelectedPrim(FPrimData(KPT_Sphyl, PrimIndex));
+		AddSelectedPrim(PrimData);
+		while( OverlapsExistingPrim(PrimData) )
+		{
+			TranslateSelectedPrims(OverlapNudge);
+		}
 	}
 
 	Viewport->RefreshViewport();
@@ -1243,6 +1357,13 @@ void FStaticMeshEditor::OnConvertBoxToConvexCollision()
 
 				BodySetup->CreatePhysicsMeshes();
 
+				// Select the new prims
+				FKAggregateGeom* AggGeom = &StaticMesh->BodySetup->AggGeom;
+				for (int32 i = 0; i < NumBoxElems; ++i)
+				{
+					AddSelectedPrim(FPrimData(KPT_Convex, (AggGeom->ConvexElems.Num() - (i+1))));
+				}
+
 				// Mark static mesh as dirty, to help make sure it gets saved.
 				StaticMesh->MarkPackageDirty();
 
@@ -1290,9 +1411,15 @@ void FStaticMeshEditor::OnCopyCollisionFromSelectedStaticMesh()
 			// Make sure rendering is done - so we are not changing data being used by collision drawing.
 			FlushRenderingCommands();
 
-			// copy body properties from
+			// Copy body properties from
 			BodySetup->CopyBodyPropertiesFrom(SelectedMesh->BodySetup);
 
+			// Enable collision, if not already
+			if( !Viewport->GetViewportClient().IsSetShowWireframeCollisionChecked() )
+			{
+				Viewport->GetViewportClient().SetShowWireframeCollision();
+			}
+
 			//Invalidate physics data
 			BodySetup->InvalidatePhysicsData();
 
@@ -1441,6 +1568,12 @@ void FStaticMeshEditor::DoDecomp(int32 InMaxHullCount, int32 InMaxHullVerts)
 		// Run actual util to do the work
 		DecomposeMeshToHulls(bs, Verts, Indices, InMaxHullCount, InMaxHullVerts);		
 
+		// Enable collision, if not already
+		if( !Viewport->GetViewportClient().IsSetShowWireframeCollisionChecked() )
+		{
+			Viewport->GetViewportClient().SetShowWireframeCollision();
+		}
+
 		// refresh collision change back to staticmesh components
 		RefreshCollisionChange(StaticMesh);
 
diff --git a/Engine/Source/Editor/StaticMeshEditor/Private/StaticMeshEditor.h b/Engine/Source/Editor/StaticMeshEditor/Private/StaticMeshEditor.h
--- a/Engine/Source/Editor/StaticMeshEditor/Private/StaticMeshEditor.h
+++ b/Engine/Source/Editor/StaticMeshEditor/Private/StaticMeshEditor.h
@@ -14,6 +14,7 @@ public:
 		: StaticMesh( NULL )
 		, NumLODLevels(0)
 		, MinPrimSize(0.5f)
+		, OverlapNudge(10.0f)
 	{}
 
 	~FStaticMeshEditor();
@@ -75,6 +76,7 @@ public:
 	virtual bool GetLastSelectedPrimTransform(FTransform& OutTransform) const;
 	FTransform GetPrimTransform(const FPrimData& InPrimData) const;
 	void SetPrimTransform(const FPrimData& InPrimData, const FTransform& InPrimTransform) const;
+	bool OverlapsExistingPrim(const FPrimData& InPrimData) const;
 
 	virtual int32 GetNumTriangles( int32 LODLevel = 0 ) const override;
 	virtual int32 GetNumVertices( int32 LODLevel = 0 ) const override;
@@ -295,7 +297,8 @@ private:
 	TArray<FPrimData> SelectedPrims;
 
 	/** Misc consts */
-	const float	MinPrimSize;
+	const float MinPrimSize;
+	const FVector OverlapNudge;
 
 	/**	The tab ids for all the tabs used */
 	static const FName ViewportTabId;
diff --git a/Engine/Source/Editor/StaticMeshEditor/Private/StaticMeshEditorActions.cpp b/Engine/Source/Editor/StaticMeshEditor/Private/StaticMeshEditorActions.cpp
--- a/Engine/Source/Editor/StaticMeshEditor/Private/StaticMeshEditorActions.cpp
+++ b/Engine/Source/Editor/StaticMeshEditor/Private/StaticMeshEditorActions.cpp
@@ -22,15 +22,15 @@ void FStaticMeshEditorCommands::RegisterCommands()
 	UI_COMMAND( SetShowBinormals, "Binormals", "Toggles display of vertex binormals (orthogonal vector to normal and tangent) in the Preview Pane.", EUserInterfaceActionType::ToggleButton, FInputGesture() );
 	UI_COMMAND( SetShowPivot, "Show Pivot", "Display the pivot location of the static mesh.", EUserInterfaceActionType::ToggleButton, FInputGesture() );	
 
-	UI_COMMAND( CreateDOP10X, "10DOP-X Simplified Collision", "Generates a new axis-aligned box collision mesh with the 4 X-axis aligned edges beveled (10 total sides) encompassing the static mesh.", EUserInterfaceActionType::Button, FInputGesture() );
-	UI_COMMAND( CreateDOP10Y, "10DOP-Y Simplified Collision", "Generates a new axis-aligned box collision mesh with the 4 Y-axis aligned edges beveled (10 total sides) encompassing the static mesh.", EUserInterfaceActionType::Button, FInputGesture() );
-	UI_COMMAND( CreateDOP10Z, "10DOP-Z Simplified Collision", "Generates a new axis-aligned box collision mesh with the 4 Z-axis aligned edges beveled (10 total sides) encompassing the static mesh.", EUserInterfaceActionType::Button, FInputGesture() );
-	UI_COMMAND( CreateDOP18, "18DOP Simplified Collision", "Generates a new axis-aligned box collision mesh with all edges beveled (18 total sides) encompassing the static mesh.", EUserInterfaceActionType::Button, FInputGesture() );
-	UI_COMMAND( CreateDOP26, "26DOP Simplified Collision", "Generates a new axis-aligned box collision mesh with all edges and corners beveled (26 total sides) encompassing the static mesh.", EUserInterfaceActionType::Button, FInputGesture() );
-
-	UI_COMMAND( CreateBoxCollision, "Box Simplified Collision", "Generates a new box collision mesh encompassing the static mesh.", EUserInterfaceActionType::Button, FInputGesture() );
-	UI_COMMAND( CreateSphereCollision, "Sphere Simplified Collision", "Generates a new sphere collision mesh encompassing the static mesh.", EUserInterfaceActionType::Button, FInputGesture() );
-	UI_COMMAND( CreateSphylCollision, "Capsule Simplified Collision", "Generates a new capsule collision mesh encompassing the static mesh.", EUserInterfaceActionType::Button, FInputGesture() );
+	UI_COMMAND( CreateDOP10X, "Add 10DOP-X Simplified Collision", "Generates a new axis-aligned box collision mesh with the 4 X-axis aligned edges beveled (10 total sides) encompassing the static mesh.", EUserInterfaceActionType::Button, FInputGesture() );
+	UI_COMMAND( CreateDOP10Y, "Add 10DOP-Y Simplified Collision", "Generates a new axis-aligned box collision mesh with the 4 Y-axis aligned edges beveled (10 total sides) encompassing the static mesh.", EUserInterfaceActionType::Button, FInputGesture() );
+	UI_COMMAND( CreateDOP10Z, "Add 10DOP-Z Simplified Collision", "Generates a new axis-aligned box collision mesh with the 4 Z-axis aligned edges beveled (10 total sides) encompassing the static mesh.", EUserInterfaceActionType::Button, FInputGesture() );
+	UI_COMMAND( CreateDOP18, "Add 18DOP Simplified Collision", "Generates a new axis-aligned box collision mesh with all edges beveled (18 total sides) encompassing the static mesh.", EUserInterfaceActionType::Button, FInputGesture() );
+	UI_COMMAND( CreateDOP26, "Add 26DOP Simplified Collision", "Generates a new axis-aligned box collision mesh with all edges and corners beveled (26 total sides) encompassing the static mesh.", EUserInterfaceActionType::Button, FInputGesture() );
+
+	UI_COMMAND( CreateBoxCollision, "Add Box Simplified Collision", "Generates a new box collision mesh encompassing the static mesh.", EUserInterfaceActionType::Button, FInputGesture() );
+	UI_COMMAND( CreateSphereCollision, "Add Sphere Simplified Collision", "Generates a new sphere collision mesh encompassing the static mesh.", EUserInterfaceActionType::Button, FInputGesture() );
+	UI_COMMAND( CreateSphylCollision, "Add Capsule Simplified Collision", "Generates a new capsule collision mesh encompassing the static mesh.", EUserInterfaceActionType::Button, FInputGesture() );
 	UI_COMMAND( CreateAutoConvexCollision, "Auto Convex Collision", "Opens the Auto Convex Collision Tool for generating a new convex collision mesh, or meshes.", EUserInterfaceActionType::Button, FInputGesture() );
 	UI_COMMAND( RemoveCollision, "Remove Collision", "Removes any simplified collision assigned to the static mesh.", EUserInterfaceActionType::Button, FInputGesture() );
 	UI_COMMAND( ConvertBoxesToConvex, "Convert Boxes to Convex", "Converts any simple box collision meshes to convex collision meshes.", EUserInterfaceActionType::Button, FInputGesture() );
