Fix to LOD screen sizes not working at small values, and delta shifting into negative values.

diff --git a/Engine/Source/Editor/StaticMeshEditor/Private/StaticMeshEditorTools.cpp b/Engine/Source/Editor/StaticMeshEditor/Private/StaticMeshEditorTools.cpp
--- a/Engine/Source/Editor/StaticMeshEditor/Private/StaticMeshEditorTools.cpp
+++ b/Engine/Source/Editor/StaticMeshEditor/Private/StaticMeshEditorTools.cpp
@@ -1688,15 +1688,9 @@ void FLevelOfDetailSettingsLayout::OnLODScreenSizeChanged( float NewValue, int32
 		}
 
 		// Update Display factors for further LODs
-		float Delta = NewValue - LODScreenSizes[LODIndex];
+		const float MinimumDifferenceInScreenSize = KINDA_SMALL_NUMBER;
 		LODScreenSizes[LODIndex] = NewValue;
-		for (int32 i = LODIndex + 1; i < MAX_STATIC_MESH_LODS; ++i)
-		{
-			LODScreenSizes[i] += Delta;
-		}
-
 		// Make sure we aren't trying to ovelap or have more than one LOD for a value
-		const float MinimumDifferenceInScreenSize = 0.01f;
 		for (int32 i = 1; i < MAX_STATIC_MESH_LODS; ++i)
 		{
 			float MaxValue = FMath::Clamp(LODScreenSizes[i-1] - MinimumDifferenceInScreenSize, 0.0f, 1.0f);
