Creating new static mesh collision enables collision view (if not already enabled)

#ttp 345296 - NUX: When creating collision in static mesh editor, enable collision viewing so you can see result

#branch UE4

#change Also added a param to AddPrim to specify whether we want to clear the selection, just to remove duplicate code.

diff --git a/Engine/Source/Editor/StaticMeshEditor/Private/StaticMeshEditor.h b/Engine/Source/Editor/StaticMeshEditor/Private/StaticMeshEditor.h
--- a/Engine/Source/Editor/StaticMeshEditor/Private/StaticMeshEditor.h
+++ b/Engine/Source/Editor/StaticMeshEditor/Private/StaticMeshEditor.h
@@ -63,7 +63,7 @@ public:
 
 	virtual bool IsPrimValid(const FPrimData& InPrimData) const override;
 	virtual bool HasSelectedPrims() const;
-	virtual void AddSelectedPrim(const FPrimData& InPrimData) override;
+	virtual void AddSelectedPrim(const FPrimData& InPrimData, bool bClearSelection) override;
 	virtual void RemoveSelectedPrim(const FPrimData& InPrimData) override;
 	virtual void RemoveInvalidPrims() override;
 	virtual bool IsSelectedPrim(const FPrimData& InPrimData) const override;
diff --git a/Engine/Source/Editor/StaticMeshEditor/Private/StaticMeshEditorViewportClient.cpp b/Engine/Source/Editor/StaticMeshEditor/Private/StaticMeshEditorViewportClient.cpp
--- a/Engine/Source/Editor/StaticMeshEditor/Private/StaticMeshEditorViewportClient.cpp
+++ b/Engine/Source/Editor/StaticMeshEditor/Private/StaticMeshEditorViewportClient.cpp
@@ -825,8 +825,7 @@ void FStaticMeshEditorViewportClient::ProcessClick(class FSceneView& InView, cla
 			{
 				if (!bCtrlDown)
 				{
-					StaticMeshEditorPtr.Pin()->ClearSelectedPrims();
-					StaticMeshEditorPtr.Pin()->AddSelectedPrim(CollisionProxy->PrimData);
+					StaticMeshEditorPtr.Pin()->AddSelectedPrim(CollisionProxy->PrimData, true);
 				}
 				else
 				{
@@ -835,11 +834,7 @@ void FStaticMeshEditorViewportClient::ProcessClick(class FSceneView& InView, cla
 			}
 			else
 			{
-				if (!bCtrlDown)
-				{
-					StaticMeshEditorPtr.Pin()->ClearSelectedPrims();
-				}
-				StaticMeshEditorPtr.Pin()->AddSelectedPrim(CollisionProxy->PrimData);
+				StaticMeshEditorPtr.Pin()->AddSelectedPrim(CollisionProxy->PrimData, !bCtrlDown);
 			}
 
 			// Force the widget to translate, if not already set
diff --git a/Engine/Source/Editor/StaticMeshEditor/Public/IStaticMeshEditor.h b/Engine/Source/Editor/StaticMeshEditor/Public/IStaticMeshEditor.h
--- a/Engine/Source/Editor/StaticMeshEditor/Public/IStaticMeshEditor.h
+++ b/Engine/Source/Editor/StaticMeshEditor/Public/IStaticMeshEditor.h
@@ -72,8 +72,9 @@ public:
 	 *  Adds primitive information to the selected prims list
 	 *
 	 *  @param  InPrimData			The data to add
+	 *  @param  bClearSelection		If true, clears the current selection
 	 */
-	virtual void AddSelectedPrim(const FPrimData& InPrimData) = 0;
+	virtual void AddSelectedPrim(const FPrimData& InPrimData, bool bClearSelection) = 0;
 
 	/** 
 	 *  Removes primitive information to the selected prims list
