Missing file from CL#2305555

diff --git a/Engine/Source/Editor/StaticMeshEditor/Private/StaticMeshEditor.cpp b/Engine/Source/Editor/StaticMeshEditor/Private/StaticMeshEditor.cpp
--- a/Engine/Source/Editor/StaticMeshEditor/Private/StaticMeshEditor.cpp
+++ b/Engine/Source/Editor/StaticMeshEditor/Private/StaticMeshEditor.cpp
@@ -585,7 +585,7 @@ bool FStaticMeshEditor::HasSelectedPrims() const
 	return (SelectedPrims.Num() > 0 ? true : false);
 }
 
-void FStaticMeshEditor::AddSelectedPrim(const FPrimData& InPrimData)
+void FStaticMeshEditor::AddSelectedPrim(const FPrimData& InPrimData, bool bClearSelection)
 {
 	check(IsPrimValid(InPrimData));
 
@@ -595,6 +595,10 @@ void FStaticMeshEditor::AddSelectedPrim(const FPrimData& InPrimData)
 		Viewport->GetViewportClient().SetShowWireframeCollision();
 	}
 
+	if( bClearSelection )
+	{
+		ClearSelectedPrims();
+	}
 	SelectedPrims.Add(InPrimData);	
 }
 
@@ -1140,7 +1144,7 @@ void FStaticMeshEditor::GenerateKDop(const FVector* Directions, uint32 NumDirect
 		}
 		const FPrimData PrimData = FPrimData(KPT_Convex, PrimIndex);
 		ClearSelectedPrims();
-		AddSelectedPrim(PrimData);
+		AddSelectedPrim(PrimData, true);
 		while( OverlapsExistingPrim(PrimData) )
 		{
 			TranslateSelectedPrims(OverlapNudge);
@@ -1163,7 +1167,7 @@ void FStaticMeshEditor::OnCollisionBox()
 		}
 		const FPrimData PrimData = FPrimData(KPT_Box, PrimIndex);
 		ClearSelectedPrims();
-		AddSelectedPrim(PrimData);
+		AddSelectedPrim(PrimData, true);
 		while( OverlapsExistingPrim(PrimData) )
 		{
 			TranslateSelectedPrims(OverlapNudge);
@@ -1186,7 +1190,7 @@ void FStaticMeshEditor::OnCollisionSphere()
 		}
 		const FPrimData PrimData = FPrimData(KPT_Sphere, PrimIndex);
 		ClearSelectedPrims();
-		AddSelectedPrim(PrimData);
+		AddSelectedPrim(PrimData, true);
 		while( OverlapsExistingPrim(PrimData) )
 		{
 			TranslateSelectedPrims(OverlapNudge);
@@ -1209,7 +1213,7 @@ void FStaticMeshEditor::OnCollisionSphyl()
 		}
 		const FPrimData PrimData = FPrimData(KPT_Sphyl, PrimIndex);
 		ClearSelectedPrims();
-		AddSelectedPrim(PrimData);
+		AddSelectedPrim(PrimData, true);
 		while( OverlapsExistingPrim(PrimData) )
 		{
 			TranslateSelectedPrims(OverlapNudge);
@@ -1361,7 +1365,7 @@ void FStaticMeshEditor::OnConvertBoxToConvexCollision()
 				FKAggregateGeom* AggGeom = &StaticMesh->BodySetup->AggGeom;
 				for (int32 i = 0; i < NumBoxElems; ++i)
 				{
-					AddSelectedPrim(FPrimData(KPT_Convex, (AggGeom->ConvexElems.Num() - (i+1))));
+					AddSelectedPrim(FPrimData(KPT_Convex, (AggGeom->ConvexElems.Num() - (i+1))), false);
 				}
 
 				// Mark static mesh as dirty, to help make sure it gets saved.
