fixed bug with bad src channel for auto lightmap UVs

diff --git a/Engine/Source/Developer/MeshUtilities/Private/MeshUtilities.cpp b/Engine/Source/Developer/MeshUtilities/Private/MeshUtilities.cpp
--- a/Engine/Source/Developer/MeshUtilities/Private/MeshUtilities.cpp
+++ b/Engine/Source/Developer/MeshUtilities/Private/MeshUtilities.cpp
@@ -2343,6 +2343,11 @@ bool FMeshUtilities::BuildStaticMesh(
 			// Generate lightmap UVs
 			if( SrcModel.BuildSettings.bGenerateLightmapUVs )
 			{
+				if( RawMesh.WedgeTexCoords[ SrcModel.BuildSettings.SrcLightmapIndex ].Num() == 0 )
+				{
+					SrcModel.BuildSettings.SrcLightmapIndex = 0;
+				}
+
 				FLayoutUV Packer( &RawMesh, SrcModel.BuildSettings.SrcLightmapIndex, SrcModel.BuildSettings.DstLightmapIndex, SrcModel.BuildSettings.MinLightmapResolution );
 
 				Packer.FindCharts( OverlappingCorners );
diff --git a/Engine/Source/Editor/StaticMeshEditor/Private/StaticMeshEditorTools.cpp b/Engine/Source/Editor/StaticMeshEditor/Private/StaticMeshEditorTools.cpp
--- a/Engine/Source/Editor/StaticMeshEditor/Private/StaticMeshEditorTools.cpp
+++ b/Engine/Source/Editor/StaticMeshEditor/Private/StaticMeshEditorTools.cpp
@@ -427,7 +427,7 @@ void FMeshBuildSettingsLayout::GenerateChildContent( IDetailChildrenBuilder& Chi
 			SNew(SSpinBox<int32>)
 			.Font( IDetailLayoutBuilder::GetDetailFont() )
 			.MinValue(0)
-			.MaxValue(8)
+			.MaxValue(7)
 			.Value(this, &FMeshBuildSettingsLayout::GetSrcLightmapIndex)
 			.OnValueChanged(this, &FMeshBuildSettingsLayout::OnSrcLightmapIndexChanged)
 		];
@@ -446,7 +446,7 @@ void FMeshBuildSettingsLayout::GenerateChildContent( IDetailChildrenBuilder& Chi
 			SNew(SSpinBox<int32>)
 			.Font( IDetailLayoutBuilder::GetDetailFont() )
 			.MinValue(0)
-			.MaxValue(8)
+			.MaxValue(7)
 			.Value(this, &FMeshBuildSettingsLayout::GetDstLightmapIndex)
 			.OnValueChanged(this, &FMeshBuildSettingsLayout::OnDstLightmapIndexChanged)
 		];
