UE-3610 UE-4199: Fix issues with collision/rendering not updating when editing static mesh collision

diff --git a/Engine/Source/Editor/StaticMeshEditor/Private/StaticMeshEditor.cpp b/Engine/Source/Editor/StaticMeshEditor/Private/StaticMeshEditor.cpp
--- a/Engine/Source/Editor/StaticMeshEditor/Private/StaticMeshEditor.cpp
+++ b/Engine/Source/Editor/StaticMeshEditor/Private/StaticMeshEditor.cpp
@@ -714,6 +714,9 @@ void FStaticMeshEditor::DuplicateSelectedPrims(const FVector* InOffset)
 
 void FStaticMeshEditor::TranslateSelectedPrims(const FVector& InDrag)
 {
+	check(StaticMesh->BodySetup);
+	StaticMesh->BodySetup->InvalidatePhysicsData();
+
 	for (int32 PrimIdx = 0; PrimIdx < SelectedPrims.Num(); PrimIdx++)
 	{
 		const FPrimData& PrimData = SelectedPrims[PrimIdx];
@@ -726,10 +729,16 @@ void FStaticMeshEditor::TranslateSelectedPrims(const FVector& InDrag)
 
 		SetPrimTransform(PrimData, PrimTransform);
 	}
+
+	// refresh collision change back to staticmesh components
+	RefreshCollisionChange(StaticMesh);
 }
 
 void FStaticMeshEditor::RotateSelectedPrims(const FRotator& InRot)
 {
+	check(StaticMesh->BodySetup);
+	StaticMesh->BodySetup->InvalidatePhysicsData();
+
 	const FQuat DeltaQ = InRot.Quaternion();
 
 	for (int32 PrimIdx = 0; PrimIdx < SelectedPrims.Num(); PrimIdx++)
@@ -748,11 +757,15 @@ void FStaticMeshEditor::RotateSelectedPrims(const FRotator& InRot)
 
 		SetPrimTransform(PrimData, PrimTransform);
 	}
+
+	// refresh collision change back to staticmesh components
+	RefreshCollisionChange(StaticMesh);
 }
 
 void FStaticMeshEditor::ScaleSelectedPrims(const FVector& InScale)
 {
 	check(StaticMesh->BodySetup);
+	StaticMesh->BodySetup->InvalidatePhysicsData();
 
 	FKAggregateGeom* AggGeom = &StaticMesh->BodySetup->AggGeom;
 
@@ -785,6 +798,9 @@ void FStaticMeshEditor::ScaleSelectedPrims(const FVector& InScale)
 
 		StaticMesh->bCustomizedCollision = true;	//mark the static mesh for collision customization
 	}
+
+	// refresh collision change back to staticmesh components
+	RefreshCollisionChange(StaticMesh);
 }
 
 bool FStaticMeshEditor::CalcSelectedPrimsAABB(FBox &OutBox) const
@@ -1433,8 +1449,9 @@ void FStaticMeshEditor::OnCopyCollisionFromSelectedStaticMesh()
 				Viewport->GetViewportClient().SetShowWireframeCollision();
 			}
 
-			//Invalidate physics data
+			// Invalidate physics data and create new meshes
 			BodySetup->InvalidatePhysicsData();
+			BodySetup->CreatePhysicsMeshes(); 
 
 			// Mark static mesh as dirty, to help make sure it gets saved.
 			StaticMesh->MarkPackageDirty();
diff --git a/Engine/Source/Editor/StaticMeshEditor/Private/StaticMeshEditorViewportClient.cpp b/Engine/Source/Editor/StaticMeshEditor/Private/StaticMeshEditorViewportClient.cpp
--- a/Engine/Source/Editor/StaticMeshEditor/Private/StaticMeshEditorViewportClient.cpp
+++ b/Engine/Source/Editor/StaticMeshEditor/Private/StaticMeshEditorViewportClient.cpp
@@ -412,6 +412,9 @@ void FStaticMeshEditorViewportClient::Draw(const FSceneView* View,FPrimitiveDraw
 
 	if (bShowCollision && StaticMesh->BodySetup)
 	{
+		// Ensure physics mesh is created before we try and draw it
+		StaticMesh->BodySetup->CreatePhysicsMeshes();
+
 		const FColor SelectedColor(149, 223, 157);
 		const FColor UnselectedColor(157, 149, 223);
 
diff --git a/Engine/Source/Runtime/Engine/Private/PhysicsEngine/BodySetup.cpp b/Engine/Source/Runtime/Engine/Private/PhysicsEngine/BodySetup.cpp
--- a/Engine/Source/Runtime/Engine/Private/PhysicsEngine/BodySetup.cpp
+++ b/Engine/Source/Runtime/Engine/Private/PhysicsEngine/BodySetup.cpp
@@ -232,6 +232,9 @@ void UBodySetup::ClearPhysicsMeshes()
 
 	bCreatedPhysicsMeshes = false;
 #endif
+
+	// Also clear render info
+	AggGeom.FreeRenderInfo();
 }
 
 #if WITH_PHYSX
diff --git a/Engine/Source/Runtime/Engine/Private/PhysicsEngine/PhysDrawing.cpp b/Engine/Source/Runtime/Engine/Private/PhysicsEngine/PhysDrawing.cpp
--- a/Engine/Source/Runtime/Engine/Private/PhysicsEngine/PhysDrawing.cpp
+++ b/Engine/Source/Runtime/Engine/Private/PhysicsEngine/PhysDrawing.cpp
@@ -439,7 +439,7 @@ void FKConvexElem::AddCachedSolidConvexGeom(TArray<FDynamicMeshVertex>& VertexBu
 				int32 VertIndex = indices[j];
 
 				FDynamicMeshVertex Vert1;
-				Vert1.Position = P2UVector(PVertices[VertIndex]);
+				Vert1.Position = Transform.TransformPosition( P2UVector(PVertices[VertIndex]) ); // Apply element transform to get geom in component space
 				Vert1.Color = VertexColor;
 				Vert1.SetTangents(
 					TangentX,
@@ -581,6 +581,8 @@ void FKAggregateGeom::GetAggGeom(const FTransform& Transform, const FColor Color
 			// Cache collision vertex/index buffer
 			if(!RenderInfo)
 			{
+				UE_LOG(LogTemp, Log, TEXT("Create RenderInfo %x"), this);
+
 				//@todo - parallelrendering, remove const cast
 				FKAggregateGeom& ThisGeom = const_cast<FKAggregateGeom&>(*this);
 				ThisGeom.RenderInfo = new FKConvexGeomRenderInfo();
@@ -652,6 +654,8 @@ void FKAggregateGeom::FreeRenderInfo()
 	// See if we have rendering resources to free
 	if(RenderInfo)
 	{
+		UE_LOG(LogTemp, Log, TEXT("Create RenderInfo %x"), this);
+
 		// Should always have these if RenderInfo exists
 		check(RenderInfo->VertexBuffer);
 		check(RenderInfo->IndexBuffer);
