Converted the property editor to use FText

diff --git a/Engine/Source/Editor/PropertyEditor/Public/PropertyHandle.h b/Engine/Source/Editor/PropertyEditor/Public/PropertyHandle.h
--- a/Engine/Source/Editor/PropertyEditor/Public/PropertyHandle.h
+++ b/Engine/Source/Editor/PropertyEditor/Public/PropertyHandle.h
@@ -104,14 +104,14 @@ public:
 	/**
 	 * Gets the property tool tip text.
 	 */
-	virtual FString GetToolTipText() const = 0;
+	virtual FText GetToolTipText() const = 0;
 	
 	/**
 	 * Sets the tooltip shown for this property
 	 *
 	 * @param ToolTip	The tool tip to show
 	 */
-	virtual void SetToolTipText(const FString& ToolTip) = 0;
+	virtual void SetToolTipText(const FText& ToolTip) = 0;
 
 	/**
 	 * Gets the value formatted as a string.
@@ -294,7 +294,7 @@ public:
 	/**
 	 * @return The display name of the property
 	 */
-	virtual FString GetPropertyDisplayName() const = 0;
+	virtual FText GetPropertyDisplayName() const = 0;
 	
 	/**
 	 * Resets the value to its default
@@ -342,7 +342,7 @@ public:
 	 * @param bDisplayThumbnail			Whether or not to display the thumbnail for the property (if any)
 	 * @return the name widget for this property
 	 */
-	virtual TSharedRef<SWidget> CreatePropertyNameWidget( const FString& NameOverride = TEXT(""), const FString& ToolTipOverride = TEXT(""), bool bDisplayResetToDefault = false, bool bDisplayText = true, bool bDisplayThumbnail = true ) const = 0;
+	virtual TSharedRef<SWidget> CreatePropertyNameWidget( const FText& NameOverride = FText::GetEmpty(), const FText& ToolTipOverride = FText::GetEmpty(), bool bDisplayResetToDefault = false, bool bDisplayText = true, bool bDisplayThumbnail = true ) const = 0;
 
 	/**
 	 * Creates a value widget for this property
diff --git a/Engine/Source/Editor/StaticMeshEditor/Private/StaticMeshEditorTools.cpp b/Engine/Source/Editor/StaticMeshEditor/Private/StaticMeshEditorTools.cpp
--- a/Engine/Source/Editor/StaticMeshEditor/Private/StaticMeshEditorTools.cpp
+++ b/Engine/Source/Editor/StaticMeshEditor/Private/StaticMeshEditorTools.cpp
@@ -36,11 +36,11 @@ FStaticMeshDetails::~FStaticMeshDetails()
 
 void FStaticMeshDetails::CustomizeDetails( class IDetailLayoutBuilder& DetailBuilder )
 {
-	IDetailCategoryBuilder& LODSettingsCategory = DetailBuilder.EditCategory( "LodSettings", LOCTEXT("LodSettingsCategory", "LOD Settings").ToString() );
-	IDetailCategoryBuilder& StaticMeshCategory = DetailBuilder.EditCategory( "StaticMesh", LOCTEXT("StaticMeshGeneralSettings", "Static Mesh Settings").ToString() );
-	IDetailCategoryBuilder& ImportCategory = DetailBuilder.EditCategory( "ImportSettings", LOCTEXT("ImportGeneralSettings", "Import Settings").ToString() );
+	IDetailCategoryBuilder& LODSettingsCategory = DetailBuilder.EditCategory( "LodSettings", LOCTEXT("LodSettingsCategory", "LOD Settings") );
+	IDetailCategoryBuilder& StaticMeshCategory = DetailBuilder.EditCategory( "StaticMesh", LOCTEXT("StaticMeshGeneralSettings", "Static Mesh Settings") );
+	IDetailCategoryBuilder& ImportCategory = DetailBuilder.EditCategory( "ImportSettings", LOCTEXT("ImportGeneralSettings", "Import Settings") );
 	
-	DetailBuilder.EditCategory( "Navigation", TEXT(""), ECategoryPriority::Uncommon );
+	DetailBuilder.EditCategory( "Navigation", FText::GetEmpty(), ECategoryPriority::Uncommon );
 
 	LevelOfDetailSettings = MakeShareable( new FLevelOfDetailSettingsLayout( StaticMeshEditor ) );
 
@@ -90,7 +90,7 @@ void FStaticMeshDetails::CustomizeDetails( class IDetailLayoutBuilder& DetailBui
 	check(StaticMesh);
 	if (StaticMesh->AssetImportData && StaticMesh->AssetImportData->GetClass() != UAssetImportData::StaticClass())
 	{
-		ImportCategory.AddCustomRow(LOCTEXT("ReimportStaticMesh", "Reimport Static Mesh").ToString(), true)
+		ImportCategory.AddCustomRow(LOCTEXT("ReimportStaticMesh", "Reimport Static Mesh"), true)
 			.ValueContent()
 			[
 				SNew(SButton)
@@ -315,7 +315,7 @@ void FMeshBuildSettingsLayout::GenerateHeaderRowContent( FDetailWidgetRow& NodeR
 void FMeshBuildSettingsLayout::GenerateChildContent( IDetailChildrenBuilder& ChildrenBuilder )
 {
 	{
-		ChildrenBuilder.AddChildContent( LOCTEXT("RecomputeNormals", "Recompute Normals").ToString() )
+		ChildrenBuilder.AddChildContent( LOCTEXT("RecomputeNormals", "Recompute Normals") )
 		.NameContent()
 		[
 			SNew(STextBlock)
@@ -332,7 +332,7 @@ void FMeshBuildSettingsLayout::GenerateChildContent( IDetailChildrenBuilder& Chi
 	}
 
 	{
-		ChildrenBuilder.AddChildContent( LOCTEXT("RecomputeTangents", "Recompute Tangents").ToString() )
+		ChildrenBuilder.AddChildContent( LOCTEXT("RecomputeTangents", "Recompute Tangents") )
 		.NameContent()
 		[
 			SNew(STextBlock)
@@ -348,7 +348,7 @@ void FMeshBuildSettingsLayout::GenerateChildContent( IDetailChildrenBuilder& Chi
 	}
 
 	{
-		ChildrenBuilder.AddChildContent( LOCTEXT("RemoveDegenerates", "Remove Degenerates").ToString() )
+		ChildrenBuilder.AddChildContent( LOCTEXT("RemoveDegenerates", "Remove Degenerates") )
 		.NameContent()
 		[
 			SNew(STextBlock)
@@ -364,7 +364,7 @@ void FMeshBuildSettingsLayout::GenerateChildContent( IDetailChildrenBuilder& Chi
 	}
 
 	{
-		ChildrenBuilder.AddChildContent( LOCTEXT("UseFullPrecisionUVs", "Use Full Precision UVs").ToString() )
+		ChildrenBuilder.AddChildContent( LOCTEXT("UseFullPrecisionUVs", "Use Full Precision UVs") )
 		.NameContent()
 		[
 			SNew(STextBlock)
@@ -380,7 +380,7 @@ void FMeshBuildSettingsLayout::GenerateChildContent( IDetailChildrenBuilder& Chi
 	}
 
 	{
-		ChildrenBuilder.AddChildContent( LOCTEXT("GenerateLightmapUVs", "Generate Lightmap UVs").ToString() )
+		ChildrenBuilder.AddChildContent( LOCTEXT("GenerateLightmapUVs", "Generate Lightmap UVs") )
 		.NameContent()
 		[
 			SNew(STextBlock)
@@ -396,12 +396,12 @@ void FMeshBuildSettingsLayout::GenerateChildContent( IDetailChildrenBuilder& Chi
 	}
 
 	{
-		ChildrenBuilder.AddChildContent( LOCTEXT("MinLightmapResolution", "Min Lightmap Resolution").ToString() )
+		ChildrenBuilder.AddChildContent( LOCTEXT("MinLightmapResolution", "Min Lightmap Resolution") )
 		.NameContent()
 		[
 			SNew(STextBlock)
 			.Font( IDetailLayoutBuilder::GetDetailFont() )
-			.Text(LOCTEXT("MinLightmapResolution", "Min Lightmap Resolution").ToString())
+			.Text(LOCTEXT("MinLightmapResolution", "Min Lightmap Resolution"))
 		]
 		.ValueContent()
 		[
@@ -415,12 +415,12 @@ void FMeshBuildSettingsLayout::GenerateChildContent( IDetailChildrenBuilder& Chi
 	}
 
 	{
-		ChildrenBuilder.AddChildContent( LOCTEXT("SourceLightmapIndex", "Source Lightmap Index").ToString() )
+		ChildrenBuilder.AddChildContent( LOCTEXT("SourceLightmapIndex", "Source Lightmap Index") )
 		.NameContent()
 		[
 			SNew(STextBlock)
 			.Font( IDetailLayoutBuilder::GetDetailFont() )
-			.Text(LOCTEXT("SourceLightmapIndex", "Source Lightmap Index").ToString())
+			.Text(LOCTEXT("SourceLightmapIndex", "Source Lightmap Index"))
 		]
 		.ValueContent()
 		[
@@ -434,12 +434,12 @@ void FMeshBuildSettingsLayout::GenerateChildContent( IDetailChildrenBuilder& Chi
 	}
 
 	{
-		ChildrenBuilder.AddChildContent( LOCTEXT("DestinationLightmapIndex", "Destination Lightmap Index").ToString() )
+		ChildrenBuilder.AddChildContent( LOCTEXT("DestinationLightmapIndex", "Destination Lightmap Index") )
 		.NameContent()
 		[
 			SNew(STextBlock)
 			.Font( IDetailLayoutBuilder::GetDetailFont() )
-			.Text(LOCTEXT("DestinationLightmapIndex", "Destination Lightmap Index").ToString())
+			.Text(LOCTEXT("DestinationLightmapIndex", "Destination Lightmap Index"))
 		]
 		.ValueContent()
 		[
@@ -453,7 +453,7 @@ void FMeshBuildSettingsLayout::GenerateChildContent( IDetailChildrenBuilder& Chi
 	}
 
 	{
-		ChildrenBuilder.AddChildContent(LOCTEXT("BuildScale", "Build Scale").ToString())
+		ChildrenBuilder.AddChildContent(LOCTEXT("BuildScale", "Build Scale"))
 		.NameContent()
 		[
 			SNew(STextBlock)
@@ -478,12 +478,12 @@ void FMeshBuildSettingsLayout::GenerateChildContent( IDetailChildrenBuilder& Chi
 	}
 
 	{
-		ChildrenBuilder.AddChildContent( LOCTEXT("DistanceFieldResolutionScale", "Distance Field Resolution Scale").ToString() )
+		ChildrenBuilder.AddChildContent( LOCTEXT("DistanceFieldResolutionScale", "Distance Field Resolution Scale") )
 		.NameContent()
 		[
 			SNew(STextBlock)
 			.Font( IDetailLayoutBuilder::GetDetailFont() )
-			.Text(LOCTEXT("DistanceFieldResolutionScale", "Distance Field Resolution Scale").ToString())
+			.Text(LOCTEXT("DistanceFieldResolutionScale", "Distance Field Resolution Scale"))
 		]
 		.ValueContent()
 		[
@@ -498,12 +498,12 @@ void FMeshBuildSettingsLayout::GenerateChildContent( IDetailChildrenBuilder& Chi
 	}
 		
 	{
-		ChildrenBuilder.AddChildContent( LOCTEXT("GenerateDistanceFieldAsIfTwoSided", "Generate Distance Field as if TwoSided").ToString() )
+		ChildrenBuilder.AddChildContent( LOCTEXT("GenerateDistanceFieldAsIfTwoSided", "Generate Distance Field as if TwoSided") )
 		.NameContent()
 		[
 			SNew(STextBlock)
 			.Font( IDetailLayoutBuilder::GetDetailFont() )
-			.Text(LOCTEXT("GenerateDistanceFieldAsIfTwoSided", "Generate Distance Field as if TwoSided").ToString())
+			.Text(LOCTEXT("GenerateDistanceFieldAsIfTwoSided", "Generate Distance Field as if TwoSided"))
 		]
 		.ValueContent()
 		[
@@ -514,7 +514,7 @@ void FMeshBuildSettingsLayout::GenerateChildContent( IDetailChildrenBuilder& Chi
 	}
 
 	{
-		ChildrenBuilder.AddChildContent( LOCTEXT("ApplyChanges", "Apply Changes").ToString() )
+		ChildrenBuilder.AddChildContent( LOCTEXT("ApplyChanges", "Apply Changes") )
 		.ValueContent()
 		.HAlign(HAlign_Left)
 		[
@@ -798,7 +798,7 @@ BEGIN_SLATE_FUNCTION_BUILD_OPTIMIZATION
 void FMeshReductionSettingsLayout::GenerateChildContent( IDetailChildrenBuilder& ChildrenBuilder )
 {
 	{
-		ChildrenBuilder.AddChildContent( LOCTEXT("PercentTriangles", "Percent Triangles").ToString() )
+		ChildrenBuilder.AddChildContent( LOCTEXT("PercentTriangles", "Percent Triangles") )
 		.NameContent()
 		[
 			SNew(STextBlock)
@@ -819,7 +819,7 @@ void FMeshReductionSettingsLayout::GenerateChildContent( IDetailChildrenBuilder&
 	}
 
 	{
-		ChildrenBuilder.AddChildContent( LOCTEXT("MaxDeviation", "Max Deviation").ToString() )
+		ChildrenBuilder.AddChildContent( LOCTEXT("MaxDeviation", "Max Deviation") )
 		.NameContent()
 		[
 			SNew(STextBlock)
@@ -840,7 +840,7 @@ void FMeshReductionSettingsLayout::GenerateChildContent( IDetailChildrenBuilder&
 	}
 
 	{
-		ChildrenBuilder.AddChildContent( LOCTEXT("Silhouette_MeshSimplification", "Silhouette").ToString() )
+		ChildrenBuilder.AddChildContent( LOCTEXT("Silhouette_MeshSimplification", "Silhouette") )
 		.NameContent()
 		[
 			SNew( STextBlock )
@@ -860,7 +860,7 @@ void FMeshReductionSettingsLayout::GenerateChildContent( IDetailChildrenBuilder&
 	}
 
 	{
-		ChildrenBuilder.AddChildContent( LOCTEXT("Texture_MeshSimplification", "Texture").ToString() )
+		ChildrenBuilder.AddChildContent( LOCTEXT("Texture_MeshSimplification", "Texture") )
 		.NameContent()
 		[
 			SNew( STextBlock )
@@ -880,7 +880,7 @@ void FMeshReductionSettingsLayout::GenerateChildContent( IDetailChildrenBuilder&
 	}
 
 	{
-		ChildrenBuilder.AddChildContent( LOCTEXT("Shading_MeshSimplification", "Shading").ToString() )
+		ChildrenBuilder.AddChildContent( LOCTEXT("Shading_MeshSimplification", "Shading") )
 		.NameContent()
 		[
 			SNew( STextBlock )
@@ -900,7 +900,7 @@ void FMeshReductionSettingsLayout::GenerateChildContent( IDetailChildrenBuilder&
 	}
 
 	{
-		ChildrenBuilder.AddChildContent( LOCTEXT("WeldingThreshold", "Welding Threshold").ToString() )
+		ChildrenBuilder.AddChildContent( LOCTEXT("WeldingThreshold", "Welding Threshold") )
 		.NameContent()
 		[
 			SNew(STextBlock)
@@ -921,7 +921,7 @@ void FMeshReductionSettingsLayout::GenerateChildContent( IDetailChildrenBuilder&
 	}
 
 	{
-		ChildrenBuilder.AddChildContent( LOCTEXT("RecomputeNormals", "Recompute Normals").ToString() )
+		ChildrenBuilder.AddChildContent( LOCTEXT("RecomputeNormals", "Recompute Normals") )
 		.NameContent()
 		[
 			SNew(STextBlock)
@@ -938,7 +938,7 @@ void FMeshReductionSettingsLayout::GenerateChildContent( IDetailChildrenBuilder&
 	}
 
 	{
-		ChildrenBuilder.AddChildContent( LOCTEXT("HardEdgeAngle", "Hard Edge Angle").ToString() )
+		ChildrenBuilder.AddChildContent( LOCTEXT("HardEdgeAngle", "Hard Edge Angle") )
 		.NameContent()
 		[
 			SNew(STextBlock)
@@ -959,7 +959,7 @@ void FMeshReductionSettingsLayout::GenerateChildContent( IDetailChildrenBuilder&
 	}
 
 	{
-		ChildrenBuilder.AddChildContent( LOCTEXT("ApplyChanges", "Apply Changes").ToString() )
+		ChildrenBuilder.AddChildContent( LOCTEXT("ApplyChanges", "Apply Changes") )
 			.ValueContent()
 			.HAlign(HAlign_Left)
 			[
@@ -1396,11 +1396,11 @@ void FLevelOfDetailSettingsLayout::AddToDetailsPanel( IDetailLayoutBuilder& Deta
 	UStaticMesh* StaticMesh = StaticMeshEditor.GetStaticMesh();
 
 	IDetailCategoryBuilder& LODSettingsCategory =
-		DetailBuilder.EditCategory( "LodSettings", LOCTEXT("LodSettingsCategory", "LOD Settings").ToString() );
+		DetailBuilder.EditCategory( "LodSettings", LOCTEXT("LodSettingsCategory", "LOD Settings") );
 
 	
 
-	LODSettingsCategory.AddCustomRow( LOCTEXT("LODGroup", "LOD Group").ToString() )
+	LODSettingsCategory.AddCustomRow( LOCTEXT("LODGroup", "LOD Group") )
 	.NameContent()
 	[
 		SNew(STextBlock)
@@ -1416,7 +1416,7 @@ void FLevelOfDetailSettingsLayout::AddToDetailsPanel( IDetailLayoutBuilder& Deta
 		.OnSelectionChanged(this, &FLevelOfDetailSettingsLayout::OnLODGroupChanged)
 	];
 	
-	LODSettingsCategory.AddCustomRow( LOCTEXT("LODImport", "LOD Import").ToString() )
+	LODSettingsCategory.AddCustomRow( LOCTEXT("LODImport", "LOD Import") )
 		.NameContent()
 		[
 			SNew(STextBlock)
@@ -1434,7 +1434,7 @@ void FLevelOfDetailSettingsLayout::AddToDetailsPanel( IDetailLayoutBuilder& Deta
 
 	// Add Number of LODs slider.
 	const int32 MinAllowedLOD = 1;
-	LODSettingsCategory.AddCustomRow( LOCTEXT("NumberOfLODs", "Number of LODs").ToString() )
+	LODSettingsCategory.AddCustomRow( LOCTEXT("NumberOfLODs", "Number of LODs") )
 	.NameContent()
 	[
 		SNew(STextBlock)
@@ -1455,7 +1455,7 @@ void FLevelOfDetailSettingsLayout::AddToDetailsPanel( IDetailLayoutBuilder& Deta
 	];
 
 	// Auto LOD distance check box.
-	LODSettingsCategory.AddCustomRow( LOCTEXT("AutoComputeLOD", "Auto Compute LOD Distances").ToString() )
+	LODSettingsCategory.AddCustomRow( LOCTEXT("AutoComputeLOD", "Auto Compute LOD Distances") )
 	.NameContent()
 	[
 		SNew(STextBlock)
@@ -1469,7 +1469,7 @@ void FLevelOfDetailSettingsLayout::AddToDetailsPanel( IDetailLayoutBuilder& Deta
 		.OnCheckStateChanged(this, &FLevelOfDetailSettingsLayout::OnAutoLODChanged)
 	];
 
-	LODSettingsCategory.AddCustomRow( LOCTEXT("ApplyChanges", "Apply Changes").ToString() )
+	LODSettingsCategory.AddCustomRow( LOCTEXT("ApplyChanges", "Apply Changes") )
 	.ValueContent()
 	.HAlign(HAlign_Left)
 	[
@@ -1485,7 +1485,7 @@ void FLevelOfDetailSettingsLayout::AddToDetailsPanel( IDetailLayoutBuilder& Deta
 
 	bool bAdvanced = true;
 	// Allowed pixel error.
-	LODSettingsCategory.AddCustomRow( LOCTEXT("AllowedPixelError", "Auto Distance Error").ToString(), bAdvanced )
+	LODSettingsCategory.AddCustomRow( LOCTEXT("AllowedPixelError", "Auto Distance Error"), bAdvanced )
 	.NameContent()
 	[
 		SNew(STextBlock)
@@ -1563,7 +1563,7 @@ void FLevelOfDetailSettingsLayout::AddLODLevelCategories( IDetailLayoutBuilder&
 			FString CategoryName = FString(TEXT("LOD"));
 			CategoryName.AppendInt( LODIndex );
 
-			FString LODLevelString = FText::Format( LOCTEXT("LODLevel", "LOD{0}"), FText::AsNumber( LODIndex ) ).ToString();
+			FText LODLevelString = FText::Format( LOCTEXT("LODLevel", "LOD{0}"), FText::AsNumber( LODIndex ) );
 
 			IDetailCategoryBuilder& LODCategory = DetailBuilder.EditCategory( *CategoryName, LODLevelString, ECategoryPriority::Important );
 
@@ -1605,7 +1605,7 @@ void FLevelOfDetailSettingsLayout::AddLODLevelCategories( IDetailLayoutBuilder&
 			SectionSettingsWidgets[ LODIndex ] = MakeShareable( new FMeshSectionSettingsLayout( StaticMeshEditor, LODIndex ) );
 			SectionSettingsWidgets[ LODIndex ]->AddToCategory( LODCategory );
 
-			LODCategory.AddCustomRow(( LOCTEXT("ScreenSizeRow", "ScreenSize").ToString()))
+			LODCategory.AddCustomRow(( LOCTEXT("ScreenSizeRow", "ScreenSize")))
 			.NameContent()
 			[
 				SNew(STextBlock)
diff --git a/Engine/Source/Editor/UMGEditor/Private/Customizations/DetailCustomizations.cpp b/Engine/Source/Editor/UMGEditor/Private/Customizations/DetailCustomizations.cpp
--- a/Engine/Source/Editor/UMGEditor/Private/Customizations/DetailCustomizations.cpp
+++ b/Engine/Source/Editor/UMGEditor/Private/Customizations/DetailCustomizations.cpp
@@ -97,7 +97,7 @@ void FBlueprintWidgetCustomization::CreateEventCustomization( IDetailLayoutBuild
 		return;
 	}
 
-	IDetailCategoryBuilder& PropertyCategory = DetailLayout.EditCategory(FObjectEditorUtils::GetCategoryFName(Property), TEXT(""), ECategoryPriority::Uncommon);
+	IDetailCategoryBuilder& PropertyCategory = DetailLayout.EditCategory(FObjectEditorUtils::GetCategoryFName(Property), FText::GetEmpty(), ECategoryPriority::Uncommon);
 
 	IDetailPropertyRow& PropertyRow = PropertyCategory.AddProperty(DelegatePropertyHandle);
 
@@ -190,9 +190,10 @@ void FBlueprintWidgetCustomization::CreateMulticastEventCustomization(IDetailLay
 
 	TSharedRef<IPropertyHandle> DelegatePropertyHandle = DetailLayout.GetProperty(DelegateProperty->GetFName(), CastChecked<UClass>(DelegateProperty->GetOuter()));
 
-	IDetailCategoryBuilder& PropertyCategory = DetailLayout.EditCategory("Events", LOCTEXT("Events", "Events").ToString(), ECategoryPriority::Uncommon);
+	IDetailCategoryBuilder& PropertyCategory = DetailLayout.EditCategory("Events", LOCTEXT("Events", "Events"), ECategoryPriority::Uncommon);
 
-	PropertyCategory.AddCustomRow(DelegateProperty->GetName())
+	FText DelegatePropertyName = FText::FromString(DelegateProperty->GetName());
+	PropertyCategory.AddCustomRow(DelegatePropertyName)
 	.NameContent()
 	[
 		SNew(SHorizontalBox)
@@ -210,7 +211,7 @@ void FBlueprintWidgetCustomization::CreateMulticastEventCustomization(IDetailLay
 		.VAlign(VAlign_Center)
 		[
 			SNew(STextBlock)
-			.Text(DelegateProperty->GetName())
+			.Text(DelegatePropertyName)
 		]
 	]
 	.ValueContent()
@@ -235,7 +236,7 @@ void FBlueprintWidgetCustomization::CustomizeDetails( IDetailLayoutBuilder& Deta
 			if ( Widget->Slot )
 			{
 				UClass* SlotClass = Widget->Slot->GetClass();
-				FString LayoutCatName = FString(TEXT("Slot (")) + SlotClass->GetDisplayNameText().ToString() + FString(TEXT(")"));
+				FText LayoutCatName = FText::Format(LOCTEXT("SlotNameFmt", "Slot ({0})"), SlotClass->GetDisplayNameText());
 
 				DetailLayout.EditCategory(LayoutCategoryKey, LayoutCatName, ECategoryPriority::TypeSpecific);
 			}
diff --git a/Engine/Source/Editor/UnrealEd/Private/Settings/LevelEditorPlaySettingsCustomization.h b/Engine/Source/Editor/UnrealEd/Private/Settings/LevelEditorPlaySettingsCustomization.h
--- a/Engine/Source/Editor/UnrealEd/Private/Settings/LevelEditorPlaySettingsCustomization.h
+++ b/Engine/Source/Editor/UnrealEd/Private/Settings/LevelEditorPlaySettingsCustomization.h
@@ -46,7 +46,7 @@ public:
 					+ SVerticalBox::Slot()
 						.AutoHeight()
 						[
-							InWindowPositionProperty->CreatePropertyNameWidget(LOCTEXT("WindowPosXLabel", "Left Position").ToString())
+							InWindowPositionProperty->CreatePropertyNameWidget(LOCTEXT("WindowPosXLabel", "Left Position"))
 						]
 
 					+ SVerticalBox::Slot()
@@ -66,7 +66,7 @@ public:
 					+ SVerticalBox::Slot()
 						.AutoHeight()
 						[
-							InWindowPositionProperty->CreatePropertyNameWidget(LOCTEXT("TopPositionLabel", "Top Position").ToString())
+							InWindowPositionProperty->CreatePropertyNameWidget(LOCTEXT("TopPositionLabel", "Top Position"))
 						]
 
 					+ SVerticalBox::Slot()
@@ -88,7 +88,7 @@ public:
 				.AutoWidth()
 				.VAlign(VAlign_Bottom)
 				[
-					InWindowPositionProperty->CreatePropertyNameWidget(LOCTEXT("CenterWindowLabel", "Always center window to screen").ToString())
+					InWindowPositionProperty->CreatePropertyNameWidget(LOCTEXT("CenterWindowLabel", "Always center window to screen"))
 				]
 		];
 	}
@@ -152,7 +152,7 @@ public:
 					+ SVerticalBox::Slot()
 						.AutoHeight()
 						[
-							WindowWidthProperty->CreatePropertyNameWidget(LOCTEXT("WindowWidthLabel", "Window Width").ToString())
+							WindowWidthProperty->CreatePropertyNameWidget(LOCTEXT("WindowWidthLabel", "Window Width"))
 						]
 
 					+ SVerticalBox::Slot()
@@ -171,7 +171,7 @@ public:
 					+ SVerticalBox::Slot()
 						.AutoHeight()
 						[
-							WindowHeightProperty->CreatePropertyNameWidget(LOCTEXT("WindowHeightLabel", "Window Height").ToString())
+							WindowHeightProperty->CreatePropertyNameWidget(LOCTEXT("WindowHeightLabel", "Window Height"))
 						]
 
 					+ SVerticalBox::Slot()
@@ -310,7 +310,7 @@ public:
 			WindowPositionHandle->MarkHiddenByCustomization();
 			CenterNewWindowHandle->MarkHiddenByCustomization();
 
-			PlayInNewWindowCategory.AddCustomRow(LOCTEXT("NewWindowSizeRow", "New Window Size").ToString(), false)
+			PlayInNewWindowCategory.AddCustomRow(LOCTEXT("NewWindowSizeRow", "New Window Size"), false)
 				.NameContent()
 				[
 					SNew(STextBlock)
@@ -324,7 +324,7 @@ public:
 					SNew(SScreenResolutionCustomization, &LayoutBuilder, WindowHeightHandle, WindowWidthHandle)
 				];
 
-			PlayInNewWindowCategory.AddCustomRow(LOCTEXT("NewWindowPositionRow", "New Window Position").ToString(), false)
+			PlayInNewWindowCategory.AddCustomRow(LOCTEXT("NewWindowPositionRow", "New Window Position"), false)
 				.NameContent()
 				[
 					SNew(STextBlock)
@@ -349,7 +349,7 @@ public:
 			WindowHeightHandle->MarkHiddenByCustomization();
 			WindowWidthHandle->MarkHiddenByCustomization();
 
-			PlayInStandaloneCategory.AddCustomRow(LOCTEXT("PlayInStandaloneWindowDetails", "Standalone Window Size").ToString(), false)
+			PlayInStandaloneCategory.AddCustomRow(LOCTEXT("PlayInStandaloneWindowDetails", "Standalone Window Size"), false)
 				.NameContent()
 				[
 					SNew(STextBlock)
@@ -368,7 +368,7 @@ public:
 
 			DisableStandaloneSoundProperty->MarkHiddenByCustomization();
 
-			PlayInStandaloneCategory.AddCustomRow(LOCTEXT("AdditionalStandaloneDetails", "Additional Options").ToString(), true)
+			PlayInStandaloneCategory.AddCustomRow(LOCTEXT("AdditionalStandaloneDetails", "Additional Options"), true)
 				.NameContent()
 				[
 					SNew(STextBlock)
@@ -392,7 +392,7 @@ public:
 						.VAlign(VAlign_Center)
 						.AutoWidth()
 						[
-							DisableStandaloneSoundProperty->CreatePropertyNameWidget(LOCTEXT("DisableStandaloneSoundLabel", "Disable Sound (-nosound)").ToString())
+							DisableStandaloneSoundProperty->CreatePropertyNameWidget(LOCTEXT("DisableStandaloneSoundLabel", "Disable Sound (-nosound)"))
 						]
 		
 				];
@@ -403,15 +403,15 @@ public:
 		{
 			// Number of players
 			NetworkCategory.AddProperty("PlayNumberOfClients")
-				.DisplayName(TEXT("Number of Players"))
+				.DisplayName(LOCTEXT("NumberOfPlayersLabel", "Number of Players"))
 				.IsEnabled(TAttribute<bool>::Create(TAttribute<bool>::FGetter::CreateSP(this, &FLevelEditorPlaySettingsCustomization::HandlePlayNumberOfClientsIsEnabled)));
 
 			NetworkCategory.AddProperty("AdditionalServerGameOptions")
-				.DisplayName(TEXT("Server Game Options"))
+				.DisplayName(LOCTEXT("ServerGameOptionsLabel", "Server Game Options"))
 				.IsEnabled(TAttribute<bool>::Create(TAttribute<bool>::FGetter::CreateSP(this, &FLevelEditorPlaySettingsCustomization::HandleGameOptionsIsEnabled)));
 
 			NetworkCategory.AddProperty("PlayNetDedicated")
-				.DisplayName(TEXT("Run Dedicated Server"))
+				.DisplayName(LOCTEXT("RunDedicatedServerLabel", "Run Dedicated Server"))
 				.IsEnabled(TAttribute<bool>::Create(TAttribute<bool>::FGetter::CreateSP(this, &FLevelEditorPlaySettingsCustomization::HandlePlayNetDedicatedPropertyIsEnabled)));
 
 			// client window size
@@ -422,7 +422,7 @@ public:
 			WindowWidthHandle->MarkHiddenByCustomization();
 
 			NetworkCategory.AddProperty("RouteGamepadToSecondWindow")
-				.DisplayName(TEXT("Route 1st Gamepad to 2nd Client"))
+				.DisplayName(LOCTEXT("RouteGamepadToSecondWindowLabel", "Route 1st Gamepad to 2nd Client"))
 				.IsEnabled(TAttribute<bool>::Create(TAttribute<bool>::FGetter::CreateSP(this, &FLevelEditorPlaySettingsCustomization::HandleRerouteInputToSecondWindowEnabled)))
 				.Visibility(TAttribute<EVisibility>::Create(TAttribute<EVisibility>::FGetter::CreateSP(this, &FLevelEditorPlaySettingsCustomization::HandleRerouteInputToSecondWindowVisibility)));
 			
@@ -430,12 +430,12 @@ public:
 			if (GEditor && GEditor->bAllowMultiplePIEWorlds)
 			{
 				NetworkCategory.AddProperty("RunUnderOneProcess")
-					.DisplayName(TEXT("Use Single Process"));
+					.DisplayName(LOCTEXT("RunUnderOneProcessEnabledLabel", "Use Single Process"));
 			}
 			else
 			{
 				NetworkCategory.AddProperty("RunUnderOneProcess")
-					.DisplayName( TEXT("Run Under One Process is disabled.") )
+					.DisplayName( LOCTEXT("RunUnderOneProcessDisabledLabel", "Run Under One Process is disabled.") )
 					.Visibility( EVisibility::Collapsed )
 					.IsEnabled( false );
 			}
@@ -443,16 +443,16 @@ public:
 			// Net Mode
 			NetworkCategory.AddProperty("PlayNetMode")
 				.Visibility(TAttribute<EVisibility>::Create(TAttribute<EVisibility>::FGetter::CreateSP(this, &FLevelEditorPlaySettingsCustomization::HandlePlayNetModeVisibility)))
-				.DisplayName(TEXT("Editor Multiplayer Mode"));
+				.DisplayName(LOCTEXT("PlayNetModeLabel", "Editor Multiplayer Mode"));
 
 			NetworkCategory.AddProperty("AdditionalLaunchOptions")
-				.DisplayName(TEXT("Command Line Arguments"))
+				.DisplayName(LOCTEXT("AdditionalLaunchOptionsLabel", "Command Line Arguments"))
 				.Visibility(TAttribute<EVisibility>::Create(TAttribute<EVisibility>::FGetter::CreateSP(this, &FLevelEditorPlaySettingsCustomization::HandleCmdLineVisibility)));
 
-			NetworkCategory.AddCustomRow(LOCTEXT("PlayInNetworkWindowDetails", "Multiplayer Window Size").ToString(), false)
+			NetworkCategory.AddCustomRow(LOCTEXT("PlayInNetworkWindowDetails", "Multiplayer Window Size"), false)
 				.NameContent()
 				[
-					WindowHeightHandle->CreatePropertyNameWidget(LOCTEXT("ClientWindowSizeName", "Multiplayer Window Size (in pixels)").ToString(), LOCTEXT("ClientWindowSizeTooltip", "Width and Height to use when spawning additional windows.").ToString())
+					WindowHeightHandle->CreatePropertyNameWidget(LOCTEXT("ClientWindowSizeName", "Multiplayer Window Size (in pixels)"), LOCTEXT("ClientWindowSizeTooltip", "Width and Height to use when spawning additional windows."))
 				]
 				.ValueContent()
 				.MaxDesiredWidth(MaxPropertyWidth)
@@ -462,7 +462,7 @@ public:
 				.IsEnabled(TAttribute<bool>(this, &FLevelEditorPlaySettingsCustomization::HandleClientWindowSizePropertyIsEnabled))
 				.Visibility(TAttribute<EVisibility>::Create(TAttribute<EVisibility>::FGetter::CreateSP(this, &FLevelEditorPlaySettingsCustomization::HandleClientWindowSizePropertyVisibility)));
 				
-			NetworkCategory.AddCustomRow(LOCTEXT("AdditionalMultiplayerDetails", "Additional Options").ToString(), true)
+			NetworkCategory.AddCustomRow(LOCTEXT("AdditionalMultiplayerDetails", "Additional Options"), true)
 				.NameContent()
 				[
 					SNew(STextBlock)
diff --git a/Engine/Source/Editor/UnrealEd/Private/Settings/ProjectPackagingSettingsCustomization.h b/Engine/Source/Editor/UnrealEd/Private/Settings/ProjectPackagingSettingsCustomization.h
--- a/Engine/Source/Editor/UnrealEd/Private/Settings/ProjectPackagingSettingsCustomization.h
+++ b/Engine/Source/Editor/UnrealEd/Private/Settings/ProjectPackagingSettingsCustomization.h
@@ -154,7 +154,7 @@ protected:
 
 			PopulateCultureList();
 
-			PackagingCategory.AddCustomRow("CulturesToStage", true)
+			PackagingCategory.AddCustomRow(LOCTEXT("CulturesToStageLabel", "Cultures To Stage"), true)
 				.NameContent()
 				.HAlign(HAlign_Fill)
 				.VAlign(VAlign_Top)
diff --git a/Engine/Source/Editor/UnrealEd/Public/GameModeInfoCustomizer.h b/Engine/Source/Editor/UnrealEd/Public/GameModeInfoCustomizer.h
--- a/Engine/Source/Editor/UnrealEd/Public/GameModeInfoCustomizer.h
+++ b/Engine/Source/Editor/UnrealEd/Public/GameModeInfoCustomizer.h
@@ -151,7 +151,7 @@ public:
 		];
 
 		static FName SelectedGameModeDetailsName(TEXT("SelectedGameModeDetails"));		
-		IDetailGroup& Group = CategoryBuilder.AddGroup(SelectedGameModeDetailsName, LOCTEXT("SelectedGameModeDetails", "Selected GameMode").ToString());
+		IDetailGroup& Group = CategoryBuilder.AddGroup(SelectedGameModeDetailsName, LOCTEXT("SelectedGameModeDetails", "Selected GameMode"));
 
 
 		// Then add rows to show key properties and let you edit them
diff --git a/Engine/Source/Editor/WorldBrowser/Private/StreamingLevels/StreamingLevelCustomization.cpp b/Engine/Source/Editor/WorldBrowser/Private/StreamingLevels/StreamingLevelCustomization.cpp
--- a/Engine/Source/Editor/WorldBrowser/Private/StreamingLevels/StreamingLevelCustomization.cpp
+++ b/Engine/Source/Editor/WorldBrowser/Private/StreamingLevels/StreamingLevelCustomization.cpp
@@ -44,7 +44,7 @@ void FStreamingLevelCustomization::CustomizeDetails(IDetailLayoutBuilder& Detail
 	LevelRotationProperty = LevelTransformProperty->GetChildHandle("Rotation");
 
 	// Add Position property
-	LevelStreamingCategory.AddCustomRow(LOCTEXT("Position", "Position").ToString())
+	LevelStreamingCategory.AddCustomRow(LOCTEXT("Position", "Position"))
 		.NameContent()
 		[
 			SNew(STextBlock)
@@ -68,7 +68,7 @@ void FStreamingLevelCustomization::CustomizeDetails(IDetailLayoutBuilder& Detail
 
 	
 	// Add Yaw Rotation property
-	LevelStreamingCategory.AddCustomRow(LOCTEXT("Rotation", "Rotation").ToString())
+	LevelStreamingCategory.AddCustomRow(LOCTEXT("Rotation", "Rotation"))
 		.NameContent()
 		[
 			SNew(STextBlock)
