#UE4

- Fixing compile error

diff --git a/Engine/Source/Editor/StaticMeshEditor/Private/StaticMeshEditorTools.cpp b/Engine/Source/Editor/StaticMeshEditor/Private/StaticMeshEditorTools.cpp
--- a/Engine/Source/Editor/StaticMeshEditor/Private/StaticMeshEditorTools.cpp
+++ b/Engine/Source/Editor/StaticMeshEditor/Private/StaticMeshEditorTools.cpp
@@ -351,7 +351,7 @@ void FMeshBuildSettingsLayout::GenerateChildContent( IDetailChildrenBuilder& Chi
 
 	if (GEnableMikkTSpaceCVar.GetValueOnAnyThread())
 	{
-		ChildrenBuilder.AddChildContent( LOCTEXT("UseMikkTSpace", "Use MikkTSpace").ToString() )
+		ChildrenBuilder.AddChildContent( LOCTEXT("UseMikkTSpace", "Use MikkTSpace") )
 		.NameContent()
 		[
 			SNew(STextBlock)
