UE-6159
Changing Mesh in Static Mesh Editor after adding collision Crashes

diff --git a/Engine/Source/Editor/StaticMeshEditor/Private/StaticMeshEditor.cpp b/Engine/Source/Editor/StaticMeshEditor/Private/StaticMeshEditor.cpp
--- a/Engine/Source/Editor/StaticMeshEditor/Private/StaticMeshEditor.cpp
+++ b/Engine/Source/Editor/StaticMeshEditor/Private/StaticMeshEditor.cpp
@@ -1547,6 +1547,10 @@ void FStaticMeshEditor::OnChangeMesh()
 
 		SetEditorMesh(SelectedMesh);
 
+		// Clear selections made on previous mesh
+		ClearSelectedPrims();
+		GetSelectedEdges().Empty();
+
 		if(SocketManager.IsValid())
 		{
 			SocketManager->UpdateStaticMesh();
