Allow continuous mouse wheel zoom in static mesh editor.

diff --git a/Engine/Source/Editor/StaticMeshEditor/Private/StaticMeshEditorViewportClient.cpp b/Engine/Source/Editor/StaticMeshEditor/Private/StaticMeshEditorViewportClient.cpp
--- a/Engine/Source/Editor/StaticMeshEditor/Private/StaticMeshEditorViewportClient.cpp
+++ b/Engine/Source/Editor/StaticMeshEditor/Private/StaticMeshEditorViewportClient.cpp
@@ -1092,8 +1092,12 @@ void FStaticMeshEditorViewportClient::PerspectiveCameraMoved()
 	FEditorViewportClient::PerspectiveCameraMoved();
 
 	// The static mesh editor saves the camera position in terms of an orbit camera, so ensure 
-	// that orbit mode is enabled before we store the current transform information
+	// that orbit mode is enabled before we store the current transform information.
+	// ToggleOrbitCamera affects the stored camera location, and information is lost in the transformation,
+	// so remember the current position here so it can be restored afterwards.
 	const bool bWasOrbit = bUsingOrbitCamera;
+	const FVector OldCameraLocation = GetViewLocation();
+	const FRotator OldCameraRotation = GetViewRotation();
 	ToggleOrbitCamera(true);
 
 	const FVector OrbitPoint = GetLookAtLocation();
@@ -1105,6 +1109,8 @@ void FStaticMeshEditorViewportClient::PerspectiveCameraMoved()
 		);
 
 	ToggleOrbitCamera(bWasOrbit);
+	SetViewLocation(OldCameraLocation);
+	SetViewRotation(OldCameraRotation);
 }
 
 void FStaticMeshEditorViewportClient::SetPreviewMesh(UStaticMesh* InStaticMesh, UStaticMeshComponent* InStaticMeshComponent)
