Slate - ESlateCheckBoxState has been renamed to ECheckBoxState.

diff --git a/Engine/Source/Editor/ReferenceViewer/Private/SReferenceViewer.h b/Engine/Source/Editor/ReferenceViewer/Private/SReferenceViewer.h
--- a/Engine/Source/Editor/ReferenceViewer/Private/SReferenceViewer.h
+++ b/Engine/Source/Editor/ReferenceViewer/Private/SReferenceViewer.h
@@ -61,13 +61,13 @@ private:
 
 	void OnUpdateHistoryData(FReferenceViewerHistoryData& HistoryData) const;
 	
-	void OnSearchDepthEnabledChanged( ESlateCheckBoxState::Type NewState );
-	ESlateCheckBoxState::Type IsSearchDepthEnabledChecked() const;
+	void OnSearchDepthEnabledChanged( ECheckBoxState NewState );
+	ECheckBoxState IsSearchDepthEnabledChecked() const;
 	int32 GetSearchDepthCount() const;
 	void OnSearchDepthCommitted(int32 NewValue);
 
-	void OnSearchBreadthEnabledChanged( ESlateCheckBoxState::Type NewState );
-	ESlateCheckBoxState::Type IsSearchBreadthEnabledChecked() const;
+	void OnSearchBreadthEnabledChanged( ECheckBoxState NewState );
+	ECheckBoxState IsSearchBreadthEnabledChecked() const;
 	int32 GetSearchBreadthCount() const;
 	void OnSearchBreadthCommitted(int32 NewValue);
 
diff --git a/Engine/Source/Editor/Sequencer/Private/SSequencer.cpp b/Engine/Source/Editor/Sequencer/Private/SSequencer.cpp
--- a/Engine/Source/Editor/Sequencer/Private/SSequencer.cpp
+++ b/Engine/Source/Editor/Sequencer/Private/SSequencer.cpp
@@ -463,14 +463,14 @@ EVisibility SSequencer::OnGetCleanViewVisibility() const
 	return Sequencer.Pin()->IsLevelEditorSequencer() ? EVisibility::Visible : EVisibility::Collapsed;
 }
 
-ESlateCheckBoxState::Type SSequencer::OnGetCleanViewCheckState() const
+ECheckBoxState SSequencer::OnGetCleanViewCheckState() const
 {
-	return CleanViewEnabled.Get() ? ESlateCheckBoxState::Checked : ESlateCheckBoxState::Unchecked;
+	return CleanViewEnabled.Get() ? ECheckBoxState::Checked : ECheckBoxState::Unchecked;
 }
 
-void SSequencer::OnCleanViewChecked( ESlateCheckBoxState::Type InState ) 
+void SSequencer::OnCleanViewChecked( ECheckBoxState InState ) 
 {
-	OnToggleCleanView.ExecuteIfBound( InState == ESlateCheckBoxState::Checked ? true : false );
+	OnToggleCleanView.ExecuteIfBound( InState == ECheckBoxState::Checked ? true : false );
 }
 
 TOptional<float> SSequencer::OnGetSnapInterval() const
diff --git a/Engine/Source/Editor/Sequencer/Private/SSequencer.h b/Engine/Source/Editor/Sequencer/Private/SSequencer.h
--- a/Engine/Source/Editor/Sequencer/Private/SSequencer.h
+++ b/Engine/Source/Editor/Sequencer/Private/SSequencer.h
@@ -103,14 +103,14 @@ private:
 	/**
 	 * @return The state of the auto-key check box
 	 */
-	ESlateCheckBoxState::Type OnGetCleanViewCheckState() const;
+	ECheckBoxState OnGetCleanViewCheckState() const;
 	
 	/**
 	 * Called when the auto-key checkbox state changes
 	 * 
 	 * @param InState	The new state
 	 */
-	void OnCleanViewChecked( ESlateCheckBoxState::Type InState );	
+	void OnCleanViewChecked( ECheckBoxState InState );	
 
 	/**
 	* @return The value of the current snap interval.
diff --git a/Engine/Source/Editor/SourceControlWindows/Private/SSourceControlRevert.cpp b/Engine/Source/Editor/SourceControlWindows/Private/SSourceControlRevert.cpp
--- a/Engine/Source/Editor/SourceControlWindows/Private/SSourceControlRevert.cpp
+++ b/Engine/Source/Editor/SourceControlWindows/Private/SSourceControlRevert.cpp
@@ -33,14 +33,14 @@ struct FRevertCheckBoxListViewItem
 		IsModified = false;
 	}
 
-	void OnCheckStateChanged( const ESlateCheckBoxState::Type NewCheckedState )
+	void OnCheckStateChanged( const ECheckBoxState NewCheckedState )
 	{
-		IsSelected = (NewCheckedState == ESlateCheckBoxState::Checked);
+		IsSelected = (NewCheckedState == ECheckBoxState::Checked);
 	}
 
-	ESlateCheckBoxState::Type OnIsChecked() const
+	ECheckBoxState OnIsChecked() const
 	{
-		return ( IsSelected ) ? ESlateCheckBoxState::Checked : ESlateCheckBoxState::Unchecked;
+		return ( IsSelected ) ? ECheckBoxState::Checked : ECheckBoxState::Unchecked;
 	}
 
 	EVisibility OnGetModifiedStateVisibility() const
@@ -288,16 +288,16 @@ private:
 	}
 
 	/** Called when the user checks or unchecks the revert unchanged checkbox; updates the list view accordingly */
-	void RevertUnchangedToggled( const ESlateCheckBoxState::Type NewCheckedState )
+	void RevertUnchangedToggled( const ECheckBoxState NewCheckedState )
 	{
-		bRevertUnchangedFilesOnly = (NewCheckedState == ESlateCheckBoxState::Checked);
+		bRevertUnchangedFilesOnly = (NewCheckedState == ECheckBoxState::Checked);
 	}
 
 	/**
 	 * Called whenever a column header is clicked, or in the case of the dialog, also when the "Check/Uncheck All" column header
 	 * checkbox is called, because its event bubbles to the column header. 
 	 */
-	void ColumnHeaderClicked( const ESlateCheckBoxState::Type NewCheckedState )
+	void ColumnHeaderClicked( const ECheckBoxState NewCheckedState )
 	{
 		for (int32 i=0; i<ListViewItemSource.Num(); i++)
 		{
@@ -305,7 +305,7 @@ private:
 
 			if (OnGetItemsEnabled())
 			{
-				CurListViewItem->IsSelected = (NewCheckedState == ESlateCheckBoxState::Checked);
+				CurListViewItem->IsSelected = (NewCheckedState == ECheckBoxState::Checked);
 			}
 		}
 	}
diff --git a/Engine/Source/Editor/SourceControlWindows/Private/SSourceControlSubmit.cpp b/Engine/Source/Editor/SourceControlWindows/Private/SSourceControlSubmit.cpp
--- a/Engine/Source/Editor/SourceControlWindows/Private/SSourceControlSubmit.cpp
+++ b/Engine/Source/Editor/SourceControlWindows/Private/SSourceControlSubmit.cpp
@@ -84,15 +84,15 @@ public:
 
 private:
 	/** The check status of the item. */
-	ESlateCheckBoxState::Type IsChecked() const
+	ECheckBoxState IsChecked() const
 	{
-		return SubmitItemData->bIsChecked ? ESlateCheckBoxState::Checked : ESlateCheckBoxState::Unchecked;
+		return SubmitItemData->bIsChecked ? ECheckBoxState::Checked : ECheckBoxState::Unchecked;
 	}
 
 	/** Changes the check status of the item .*/
-	void OnCheckStateChanged(ESlateCheckBoxState::Type InState)
+	void OnCheckStateChanged(ECheckBoxState InState)
 	{
-		SubmitItemData->bIsChecked = InState == ESlateCheckBoxState::Checked;
+		SubmitItemData->bIsChecked = InState == ECheckBoxState::Checked;
 	}
 
 private:
@@ -258,7 +258,7 @@ public:
 		;
 
 		DialogResult = ESubmitResults::SUBMIT_CANCELED;
-		KeepCheckedOut = ESlateCheckBoxState::Unchecked;
+		KeepCheckedOut = ECheckBoxState::Unchecked;
 
 		ParentFrame.Pin()->SetWidgetToFocusOnActivate(ChangeListDescriptionTextCtrl);
 	}
@@ -312,10 +312,10 @@ private:
 	 * @return the desired toggle state for the ToggleSelectedCheckBox.
 	 * Returns Unchecked, unless all of the selected items are Checked.
 	 */
-	ESlateCheckBoxState::Type GetToggleSelectedState() const
+	ECheckBoxState GetToggleSelectedState() const
 	{
 		// Default to a Checked state
-		ESlateCheckBoxState::Type PendingState = ESlateCheckBoxState::Checked;
+		ECheckBoxState PendingState = ECheckBoxState::Checked;
 
 		TArray< TSharedPtr<TSubmitItemData> > SelectedItems = GetSelectedItems(true);
 
@@ -327,7 +327,7 @@ private:
 			{
 				// If any item in the selection is Unchecked, then represent the entire set of highlighted items as Unchecked,
 				// so that the first (user) toggle of ToggleSelectedCheckBox consistently Checks all highlighted items
-				PendingState = ESlateCheckBoxState::Unchecked;
+				PendingState = ECheckBoxState::Unchecked;
 			}
 		}
 
@@ -338,11 +338,11 @@ private:
 	 * Toggles the highlighted items.
 	 * If no items are explicitly highlighted, toggles all items in the list.
 	 */
-	void OnToggleSelectedCheckBox(ESlateCheckBoxState::Type InNewState)
+	void OnToggleSelectedCheckBox(ECheckBoxState InNewState)
 	{
 		TArray< TSharedPtr<TSubmitItemData> > SelectedItems = GetSelectedItems(true);
 
-		const bool bIsChecked = (InNewState == ESlateCheckBoxState::Checked);
+		const bool bIsChecked = (InNewState == ECheckBoxState::Checked);
 		for(auto SelectedItem = SelectedItems.CreateConstIterator(); SelectedItem; ++SelectedItem)
 		{
 			TSubmitItemData* const Item = SelectedItem->Get();
@@ -380,7 +380,7 @@ public:
 	/** Does the user want to keep the files checked out */
 	bool WantToKeepCheckedOut()
 	{
-		return KeepCheckedOut == ESlateCheckBoxState::Checked ? true : false;
+		return KeepCheckedOut == ECheckBoxState::Checked ? true : false;
 	}
 
 private:
@@ -416,13 +416,13 @@ private:
 	}
 
 	/** Called when the Keep checked out Checkbox is changed */
-	void OnCheckStateChanged_KeepCheckedOut(ESlateCheckBoxState::Type InState)
+	void OnCheckStateChanged_KeepCheckedOut(ECheckBoxState InState)
 	{
 		KeepCheckedOut = InState;
 	}
 
 	/** Get the current state of the Keep Checked Out checkbox  */
-	ESlateCheckBoxState::Type GetKeepCheckedOut() const
+	ECheckBoxState GetKeepCheckedOut() const
 	{
 		return KeepCheckedOut;
 	}
@@ -448,7 +448,7 @@ private:
 	/** Internal widgets to save having to get in multiple places*/
 	TSharedPtr<SMultiLineEditableTextBox> ChangeListDescriptionTextCtrl;
 
-	ESlateCheckBoxState::Type	KeepCheckedOut;
+	ECheckBoxState	KeepCheckedOut;
 };
 
 static void FindFilesForCheckIn(const TArray<FString>& InFilenames, TArray<FString>& OutAddFiles, TArray<FString>& OutOpenFiles)
diff --git a/Engine/Source/Editor/StaticMeshEditor/Private/StaticMeshEditorTools.cpp b/Engine/Source/Editor/StaticMeshEditor/Private/StaticMeshEditorTools.cpp
--- a/Engine/Source/Editor/StaticMeshEditor/Private/StaticMeshEditorTools.cpp
+++ b/Engine/Source/Editor/StaticMeshEditor/Private/StaticMeshEditorTools.cpp
@@ -563,39 +563,39 @@ FReply FMeshBuildSettingsLayout::OnApplyChanges()
 	return FReply::Handled();
 }
 
-ESlateCheckBoxState::Type FMeshBuildSettingsLayout::ShouldRecomputeNormals() const
+ECheckBoxState FMeshBuildSettingsLayout::ShouldRecomputeNormals() const
 {
-	return BuildSettings.bRecomputeNormals ? ESlateCheckBoxState::Checked : ESlateCheckBoxState::Unchecked;
+	return BuildSettings.bRecomputeNormals ? ECheckBoxState::Checked : ECheckBoxState::Unchecked;
 }
 
-ESlateCheckBoxState::Type FMeshBuildSettingsLayout::ShouldRecomputeTangents() const
+ECheckBoxState FMeshBuildSettingsLayout::ShouldRecomputeTangents() const
 {
-	return BuildSettings.bRecomputeTangents ? ESlateCheckBoxState::Checked : ESlateCheckBoxState::Unchecked;
+	return BuildSettings.bRecomputeTangents ? ECheckBoxState::Checked : ECheckBoxState::Unchecked;
 }
 
-ESlateCheckBoxState::Type FMeshBuildSettingsLayout::ShouldUseMikkTSpace() const
+ECheckBoxState FMeshBuildSettingsLayout::ShouldUseMikkTSpace() const
 {
-	return BuildSettings.bUseMikkTSpace ? ESlateCheckBoxState::Checked : ESlateCheckBoxState::Unchecked;
+	return BuildSettings.bUseMikkTSpace ? ECheckBoxState::Checked : ECheckBoxState::Unchecked;
 }
 
-ESlateCheckBoxState::Type FMeshBuildSettingsLayout::ShouldRemoveDegenerates() const
+ECheckBoxState FMeshBuildSettingsLayout::ShouldRemoveDegenerates() const
 {
-	return BuildSettings.bRemoveDegenerates ? ESlateCheckBoxState::Checked : ESlateCheckBoxState::Unchecked;
+	return BuildSettings.bRemoveDegenerates ? ECheckBoxState::Checked : ECheckBoxState::Unchecked;
 }
 
-ESlateCheckBoxState::Type FMeshBuildSettingsLayout::ShouldUseFullPrecisionUVs() const
+ECheckBoxState FMeshBuildSettingsLayout::ShouldUseFullPrecisionUVs() const
 {
-	return BuildSettings.bUseFullPrecisionUVs ? ESlateCheckBoxState::Checked : ESlateCheckBoxState::Unchecked;
+	return BuildSettings.bUseFullPrecisionUVs ? ECheckBoxState::Checked : ECheckBoxState::Unchecked;
 }
 
-ESlateCheckBoxState::Type FMeshBuildSettingsLayout::ShouldGenerateLightmapUVs() const
+ECheckBoxState FMeshBuildSettingsLayout::ShouldGenerateLightmapUVs() const
 {
-	return BuildSettings.bGenerateLightmapUVs ? ESlateCheckBoxState::Checked : ESlateCheckBoxState::Unchecked;
+	return BuildSettings.bGenerateLightmapUVs ? ECheckBoxState::Checked : ECheckBoxState::Unchecked;
 }
 
-ESlateCheckBoxState::Type FMeshBuildSettingsLayout::ShouldGenerateDistanceFieldAsIfTwoSided() const
+ECheckBoxState FMeshBuildSettingsLayout::ShouldGenerateDistanceFieldAsIfTwoSided() const
 {
-	return BuildSettings.bGenerateDistanceFieldAsIfTwoSided ? ESlateCheckBoxState::Checked : ESlateCheckBoxState::Unchecked;
+	return BuildSettings.bGenerateDistanceFieldAsIfTwoSided ? ECheckBoxState::Checked : ECheckBoxState::Unchecked;
 }
 
 int32 FMeshBuildSettingsLayout::GetMinLightmapResolution() const
@@ -633,9 +633,9 @@ float FMeshBuildSettingsLayout::GetDistanceFieldResolutionScale() const
 	return BuildSettings.DistanceFieldResolutionScale;
 }
 
-void FMeshBuildSettingsLayout::OnRecomputeNormalsChanged(ESlateCheckBoxState::Type NewState)
+void FMeshBuildSettingsLayout::OnRecomputeNormalsChanged(ECheckBoxState NewState)
 {
-	const bool bRecomputeNormals = (NewState == ESlateCheckBoxState::Checked) ? true : false;
+	const bool bRecomputeNormals = (NewState == ECheckBoxState::Checked) ? true : false;
 	if (BuildSettings.bRecomputeNormals != bRecomputeNormals)
 	{
 		if (FEngineAnalytics::IsAvailable())
@@ -646,9 +646,9 @@ void FMeshBuildSettingsLayout::OnRecomputeNormalsChanged(ESlateCheckBoxState::Ty
 	}
 }
 
-void FMeshBuildSettingsLayout::OnRecomputeTangentsChanged(ESlateCheckBoxState::Type NewState)
+void FMeshBuildSettingsLayout::OnRecomputeTangentsChanged(ECheckBoxState NewState)
 {
-	const bool bRecomputeTangents = (NewState == ESlateCheckBoxState::Checked) ? true : false;
+	const bool bRecomputeTangents = (NewState == ECheckBoxState::Checked) ? true : false;
 	if (BuildSettings.bRecomputeTangents != bRecomputeTangents)
 	{
 		if (FEngineAnalytics::IsAvailable())
@@ -659,18 +659,18 @@ void FMeshBuildSettingsLayout::OnRecomputeTangentsChanged(ESlateCheckBoxState::T
 	}
 }
 
-void FMeshBuildSettingsLayout::OnUseMikkTSpaceChanged(ESlateCheckBoxState::Type NewState)
+void FMeshBuildSettingsLayout::OnUseMikkTSpaceChanged(ECheckBoxState NewState)
 {
-	const bool bUseMikkTSpace = (NewState == ESlateCheckBoxState::Checked) ? true : false;
+	const bool bUseMikkTSpace = (NewState == ECheckBoxState::Checked) ? true : false;
 	if (BuildSettings.bUseMikkTSpace != bUseMikkTSpace)
 	{
 		BuildSettings.bUseMikkTSpace = bUseMikkTSpace;
 	}
 }
 
-void FMeshBuildSettingsLayout::OnRemoveDegeneratesChanged(ESlateCheckBoxState::Type NewState)
+void FMeshBuildSettingsLayout::OnRemoveDegeneratesChanged(ECheckBoxState NewState)
 {
-	const bool bRemoveDegenerates = (NewState == ESlateCheckBoxState::Checked) ? true : false;
+	const bool bRemoveDegenerates = (NewState == ECheckBoxState::Checked) ? true : false;
 	if (BuildSettings.bRemoveDegenerates != bRemoveDegenerates)
 	{
 		if (FEngineAnalytics::IsAvailable())
@@ -681,9 +681,9 @@ void FMeshBuildSettingsLayout::OnRemoveDegeneratesChanged(ESlateCheckBoxState::T
 	}
 }
 
-void FMeshBuildSettingsLayout::OnUseFullPrecisionUVsChanged(ESlateCheckBoxState::Type NewState)
+void FMeshBuildSettingsLayout::OnUseFullPrecisionUVsChanged(ECheckBoxState NewState)
 {
-	const bool bUseFullPrecisionUVs = (NewState == ESlateCheckBoxState::Checked) ? true : false;
+	const bool bUseFullPrecisionUVs = (NewState == ECheckBoxState::Checked) ? true : false;
 	if (BuildSettings.bUseFullPrecisionUVs != bUseFullPrecisionUVs)
 	{
 		if (FEngineAnalytics::IsAvailable())
@@ -694,9 +694,9 @@ void FMeshBuildSettingsLayout::OnUseFullPrecisionUVsChanged(ESlateCheckBoxState:
 	}
 }
 
-void FMeshBuildSettingsLayout::OnGenerateLightmapUVsChanged(ESlateCheckBoxState::Type NewState)
+void FMeshBuildSettingsLayout::OnGenerateLightmapUVsChanged(ECheckBoxState NewState)
 {
-	const bool bGenerateLightmapUVs = (NewState == ESlateCheckBoxState::Checked) ? true : false;
+	const bool bGenerateLightmapUVs = (NewState == ECheckBoxState::Checked) ? true : false;
 	if (BuildSettings.bGenerateLightmapUVs != bGenerateLightmapUVs)
 	{
 		if (FEngineAnalytics::IsAvailable())
@@ -707,9 +707,9 @@ void FMeshBuildSettingsLayout::OnGenerateLightmapUVsChanged(ESlateCheckBoxState:
 	}
 }
 
-void FMeshBuildSettingsLayout::OnGenerateDistanceFieldAsIfTwoSidedChanged(ESlateCheckBoxState::Type NewState)
+void FMeshBuildSettingsLayout::OnGenerateDistanceFieldAsIfTwoSidedChanged(ECheckBoxState NewState)
 {
-	const bool bGenerateDistanceFieldAsIfTwoSided = (NewState == ESlateCheckBoxState::Checked) ? true : false;
+	const bool bGenerateDistanceFieldAsIfTwoSided = (NewState == ECheckBoxState::Checked) ? true : false;
 	if (BuildSettings.bGenerateDistanceFieldAsIfTwoSided != bGenerateDistanceFieldAsIfTwoSided)
 	{
 		if (FEngineAnalytics::IsAvailable())
@@ -1047,9 +1047,9 @@ float FMeshReductionSettingsLayout::GetWeldingThreshold() const
 	return ReductionSettings.WeldingThreshold;
 }
 
-ESlateCheckBoxState::Type FMeshReductionSettingsLayout::ShouldRecalculateNormals() const
+ECheckBoxState FMeshReductionSettingsLayout::ShouldRecalculateNormals() const
 {
-	return ReductionSettings.bRecalculateNormals ? ESlateCheckBoxState::Checked : ESlateCheckBoxState::Unchecked;
+	return ReductionSettings.bRecalculateNormals ? ECheckBoxState::Checked : ECheckBoxState::Unchecked;
 }
 
 float FMeshReductionSettingsLayout::GetHardAngleThreshold() const
@@ -1100,9 +1100,9 @@ void FMeshReductionSettingsLayout::OnWeldingThresholdCommitted(float NewValue, E
 	OnWeldingThresholdChanged(NewValue);
 }
 
-void FMeshReductionSettingsLayout::OnRecalculateNormalsChanged(ESlateCheckBoxState::Type NewValue)
+void FMeshReductionSettingsLayout::OnRecalculateNormalsChanged(ECheckBoxState NewValue)
 {
-	const bool bRecalculateNormals = NewValue == ESlateCheckBoxState::Checked;
+	const bool bRecalculateNormals = NewValue == ECheckBoxState::Checked;
 	if (ReductionSettings.bRecalculateNormals != bRecalculateNormals)
 	{
 		if (FEngineAnalytics::IsAvailable())
@@ -1308,59 +1308,59 @@ void FMeshSectionSettingsLayout::OnResetMaterialToDefaultClicked(UMaterialInterf
 	CallPostEditChange();
 }
 
-ESlateCheckBoxState::Type FMeshSectionSettingsLayout::DoesSectionCastShadow(int32 SectionIndex) const
+ECheckBoxState FMeshSectionSettingsLayout::DoesSectionCastShadow(int32 SectionIndex) const
 {
 	UStaticMesh& StaticMesh = GetStaticMesh();
 	FMeshSectionInfo Info = StaticMesh.SectionInfoMap.Get(LODIndex, SectionIndex);
-	return Info.bCastShadow ? ESlateCheckBoxState::Checked : ESlateCheckBoxState::Unchecked;
+	return Info.bCastShadow ? ECheckBoxState::Checked : ECheckBoxState::Unchecked;
 }
 
-void FMeshSectionSettingsLayout::OnSectionCastShadowChanged(ESlateCheckBoxState::Type NewState, int32 SectionIndex)
+void FMeshSectionSettingsLayout::OnSectionCastShadowChanged(ECheckBoxState NewState, int32 SectionIndex)
 {
 	UStaticMesh& StaticMesh = GetStaticMesh();
 	FMeshSectionInfo Info = StaticMesh.SectionInfoMap.Get(LODIndex, SectionIndex);
-	Info.bCastShadow = (NewState == ESlateCheckBoxState::Checked) ? true : false;
+	Info.bCastShadow = (NewState == ECheckBoxState::Checked) ? true : false;
 	StaticMesh.SectionInfoMap.Set(LODIndex, SectionIndex, Info);
 	CallPostEditChange();
 }
 
-ESlateCheckBoxState::Type FMeshSectionSettingsLayout::DoesSectionCollide(int32 SectionIndex) const
+ECheckBoxState FMeshSectionSettingsLayout::DoesSectionCollide(int32 SectionIndex) const
 {
 	UStaticMesh& StaticMesh = GetStaticMesh();
 	FMeshSectionInfo Info = StaticMesh.SectionInfoMap.Get(LODIndex, SectionIndex);
-	return Info.bEnableCollision ? ESlateCheckBoxState::Checked : ESlateCheckBoxState::Unchecked;
+	return Info.bEnableCollision ? ECheckBoxState::Checked : ECheckBoxState::Unchecked;
 }
 
-void FMeshSectionSettingsLayout::OnSectionCollisionChanged(ESlateCheckBoxState::Type NewState, int32 SectionIndex)
+void FMeshSectionSettingsLayout::OnSectionCollisionChanged(ECheckBoxState NewState, int32 SectionIndex)
 {
 	UStaticMesh& StaticMesh = GetStaticMesh();
 	FMeshSectionInfo Info = StaticMesh.SectionInfoMap.Get(LODIndex, SectionIndex);
-	Info.bEnableCollision = (NewState == ESlateCheckBoxState::Checked) ? true : false;
+	Info.bEnableCollision = (NewState == ECheckBoxState::Checked) ? true : false;
 	StaticMesh.SectionInfoMap.Set(LODIndex, SectionIndex, Info);
 	CallPostEditChange();
 }
 
-ESlateCheckBoxState::Type FMeshSectionSettingsLayout::IsSectionSelected(int32 SectionIndex) const
+ECheckBoxState FMeshSectionSettingsLayout::IsSectionSelected(int32 SectionIndex) const
 {
-	ESlateCheckBoxState::Type State = ESlateCheckBoxState::Unchecked;
+	ECheckBoxState State = ECheckBoxState::Unchecked;
 	UStaticMeshComponent* Component = StaticMeshEditor.GetStaticMeshComponent();
 	if (Component)
 	{
-		State = Component->SelectedEditorSection == SectionIndex ? ESlateCheckBoxState::Checked : ESlateCheckBoxState::Unchecked;
+		State = Component->SelectedEditorSection == SectionIndex ? ECheckBoxState::Checked : ECheckBoxState::Unchecked;
 	}
 	return State;
 }
 
-void FMeshSectionSettingsLayout::OnSectionSelectedChanged(ESlateCheckBoxState::Type NewState, int32 SectionIndex)
+void FMeshSectionSettingsLayout::OnSectionSelectedChanged(ECheckBoxState NewState, int32 SectionIndex)
 {
 	UStaticMeshComponent* Component = StaticMeshEditor.GetStaticMeshComponent();
 	if (Component)
 	{
-		if (NewState == ESlateCheckBoxState::Checked)
+		if (NewState == ECheckBoxState::Checked)
 		{
 			Component->SelectedEditorSection = SectionIndex;
 		}
-		else if (NewState == ESlateCheckBoxState::Unchecked)
+		else if (NewState == ECheckBoxState::Unchecked)
 		{
 			Component->SelectedEditorSection = INDEX_NONE;
 		}
@@ -1842,17 +1842,17 @@ bool FLevelOfDetailSettingsLayout::IsAutoLODEnabled() const
 	return StaticMesh->bAutoComputeLODScreenSize;
 }
 
-ESlateCheckBoxState::Type FLevelOfDetailSettingsLayout::IsAutoLODChecked() const
+ECheckBoxState FLevelOfDetailSettingsLayout::IsAutoLODChecked() const
 {
-	return IsAutoLODEnabled() ? ESlateCheckBoxState::Checked : ESlateCheckBoxState::Unchecked;
+	return IsAutoLODEnabled() ? ECheckBoxState::Checked : ECheckBoxState::Unchecked;
 }
 
-void FLevelOfDetailSettingsLayout::OnAutoLODChanged(ESlateCheckBoxState::Type NewState)
+void FLevelOfDetailSettingsLayout::OnAutoLODChanged(ECheckBoxState NewState)
 {
 	UStaticMesh* StaticMesh = StaticMeshEditor.GetStaticMesh();
 	check(StaticMesh);
 	StaticMesh->Modify();
-	StaticMesh->bAutoComputeLODScreenSize = (NewState == ESlateCheckBoxState::Checked) ? true : false;
+	StaticMesh->bAutoComputeLODScreenSize = (NewState == ECheckBoxState::Checked) ? true : false;
 	if (!StaticMesh->bAutoComputeLODScreenSize)
 	{
 		if (StaticMesh->SourceModels.IsValidIndex(0))
diff --git a/Engine/Source/Editor/StaticMeshEditor/Private/StaticMeshEditorTools.h b/Engine/Source/Editor/StaticMeshEditor/Private/StaticMeshEditorTools.h
--- a/Engine/Source/Editor/StaticMeshEditor/Private/StaticMeshEditorTools.h
+++ b/Engine/Source/Editor/StaticMeshEditor/Private/StaticMeshEditorTools.h
@@ -140,13 +140,13 @@ private:
 	virtual bool InitiallyCollapsed() const override { return true; }
 
 	FReply OnApplyChanges();
-	ESlateCheckBoxState::Type ShouldRecomputeNormals() const;
-	ESlateCheckBoxState::Type ShouldRecomputeTangents() const;
-	ESlateCheckBoxState::Type ShouldUseMikkTSpace() const;
-	ESlateCheckBoxState::Type ShouldRemoveDegenerates() const;
-	ESlateCheckBoxState::Type ShouldUseFullPrecisionUVs() const;
-	ESlateCheckBoxState::Type ShouldGenerateLightmapUVs() const;
-	ESlateCheckBoxState::Type ShouldGenerateDistanceFieldAsIfTwoSided() const;
+	ECheckBoxState ShouldRecomputeNormals() const;
+	ECheckBoxState ShouldRecomputeTangents() const;
+	ECheckBoxState ShouldUseMikkTSpace() const;
+	ECheckBoxState ShouldRemoveDegenerates() const;
+	ECheckBoxState ShouldUseFullPrecisionUVs() const;
+	ECheckBoxState ShouldGenerateLightmapUVs() const;
+	ECheckBoxState ShouldGenerateDistanceFieldAsIfTwoSided() const;
 	int32 GetMinLightmapResolution() const;
 	int32 GetSrcLightmapIndex() const;
 	int32 GetDstLightmapIndex() const;
@@ -155,13 +155,13 @@ private:
 	TOptional<float> GetBuildScaleZ() const;
 	float GetDistanceFieldResolutionScale() const;
 
-	void OnRecomputeNormalsChanged(ESlateCheckBoxState::Type NewState);
-	void OnRecomputeTangentsChanged(ESlateCheckBoxState::Type NewState);
-	void OnUseMikkTSpaceChanged(ESlateCheckBoxState::Type NewState);
-	void OnRemoveDegeneratesChanged(ESlateCheckBoxState::Type NewState);
-	void OnUseFullPrecisionUVsChanged(ESlateCheckBoxState::Type NewState);
-	void OnGenerateLightmapUVsChanged(ESlateCheckBoxState::Type NewState);
-	void OnGenerateDistanceFieldAsIfTwoSidedChanged(ESlateCheckBoxState::Type NewState);
+	void OnRecomputeNormalsChanged(ECheckBoxState NewState);
+	void OnRecomputeTangentsChanged(ECheckBoxState NewState);
+	void OnUseMikkTSpaceChanged(ECheckBoxState NewState);
+	void OnRemoveDegeneratesChanged(ECheckBoxState NewState);
+	void OnUseFullPrecisionUVsChanged(ECheckBoxState NewState);
+	void OnGenerateLightmapUVsChanged(ECheckBoxState NewState);
+	void OnGenerateDistanceFieldAsIfTwoSidedChanged(ECheckBoxState NewState);
 	void OnMinLightmapResolutionChanged( int32 NewValue );
 	void OnSrcLightmapIndexChanged( int32 NewValue );
 	void OnDstLightmapIndexChanged( int32 NewValue );
@@ -198,7 +198,7 @@ private:
 	float GetPercentTriangles() const;
 	float GetMaxDeviation() const;
 	float GetWeldingThreshold() const;
-	ESlateCheckBoxState::Type ShouldRecalculateNormals() const;
+	ECheckBoxState ShouldRecalculateNormals() const;
 	float GetHardAngleThreshold() const;
 
 	void OnPercentTrianglesChanged(float NewValue);
@@ -206,7 +206,7 @@ private:
 	void OnMaxDeviationChanged(float NewValue);
 	void OnMaxDeviationCommitted(float NewValue, ETextCommit::Type TextCommitType);
 	void OnReductionAmountChanged(float NewValue);
-	void OnRecalculateNormalsChanged(ESlateCheckBoxState::Type NewValue);
+	void OnRecalculateNormalsChanged(ECheckBoxState NewValue);
 	void OnWeldingThresholdChanged(float NewValue);
 	void OnWeldingThresholdCommitted(float NewValue, ETextCommit::Type TextCommitType);
 	void OnHardAngleThresholdChanged(float NewValue);
@@ -244,12 +244,12 @@ private:
 	TSharedRef<SWidget> OnGenerateNameWidgetsForMaterial(UMaterialInterface* Material, int32 SlotIndex);
 	TSharedRef<SWidget> OnGenerateWidgetsForMaterial(UMaterialInterface* Material, int32 SlotIndex);
 	void OnResetMaterialToDefaultClicked(UMaterialInterface* Material, int32 SlotIndex);
-	ESlateCheckBoxState::Type DoesSectionCastShadow(int32 SectionIndex) const;
-	void OnSectionCastShadowChanged(ESlateCheckBoxState::Type NewState, int32 SectionIndex);
-	ESlateCheckBoxState::Type DoesSectionCollide(int32 SectionIndex) const;
-	void OnSectionCollisionChanged(ESlateCheckBoxState::Type NewState, int32 SectionIndex);
-	ESlateCheckBoxState::Type IsSectionSelected(int32 SectionIndex) const;
-	void OnSectionSelectedChanged(ESlateCheckBoxState::Type NewState, int32 SectionIndex);
+	ECheckBoxState DoesSectionCastShadow(int32 SectionIndex) const;
+	void OnSectionCastShadowChanged(ECheckBoxState NewState, int32 SectionIndex);
+	ECheckBoxState DoesSectionCollide(int32 SectionIndex) const;
+	void OnSectionCollisionChanged(ECheckBoxState NewState, int32 SectionIndex);
+	ECheckBoxState IsSectionSelected(int32 SectionIndex) const;
+	void OnSectionSelectedChanged(ECheckBoxState NewState, int32 SectionIndex);
 	void CallPostEditChange(UProperty* PropertyChanged=nullptr);
 	
 	IStaticMeshEditor& StaticMeshEditor;
@@ -295,8 +295,8 @@ private:
 	void OnSectionSettingsExpanded(bool bIsExpanded, int32 LODIndex);
 	void OnLODGroupChanged(TSharedPtr<FString> NewValue, ESelectInfo::Type SelectInfo);
 	bool IsAutoLODEnabled() const;
-	ESlateCheckBoxState::Type IsAutoLODChecked() const;
-	void OnAutoLODChanged(ESlateCheckBoxState::Type NewState);
+	ECheckBoxState IsAutoLODChecked() const;
+	void OnAutoLODChanged(ECheckBoxState NewState);
 	float GetPixelError() const;
 	void OnPixelErrorChanged(float NewValue);
 	void OnImportLOD(TSharedPtr<FString> NewValue, ESelectInfo::Type SelectInfo);
diff --git a/Engine/Source/Editor/TaskBrowser/Private/STaskBrowser.cpp b/Engine/Source/Editor/TaskBrowser/Private/STaskBrowser.cpp
--- a/Engine/Source/Editor/TaskBrowser/Private/STaskBrowser.cpp
+++ b/Engine/Source/Editor/TaskBrowser/Private/STaskBrowser.cpp
@@ -192,7 +192,7 @@ static bool DecryptBuffer(const uint8* SrcBuffer,const uint32 SrcLen,uint8* Dest
 	}
 	FMemory::Free(OutBuffer);
 #elif PLATFORM_LINUX
-        printf("STaskBrowser.cpp: LINUX DecryptBuffer()\n");
+		printf("STaskBrowser.cpp: LINUX DecryptBuffer()\n");
 #elif PLATFORM_WINDOWS
 	DATA_BLOB SourceBlob, EntropyBlob, FinalBlob;
 	// Set up the datablob to encrypt
@@ -859,7 +859,7 @@ void STaskBrowser::OnDatabaseFilterSelect( TSharedPtr<FString> InFilter, ESelect
 	}
 }
 
-void STaskBrowser::OnOpenOnlyChanged( const ESlateCheckBoxState::Type NewCheckedState )
+void STaskBrowser::OnOpenOnlyChanged( const ECheckBoxState NewCheckedState )
 {
 	// Refresh the GUI.  It will apply the updated filters.
 	RefreshGUI( ETaskBrowserGUIRefreshOptions::RebuildTaskList );
@@ -868,12 +868,12 @@ void STaskBrowser::OnOpenOnlyChanged( const ESlateCheckBoxState::Type NewChecked
 	{
 		FTaskBrowserSettings TBSettings;
 		TBSettings.LoadSettings();
-		TBSettings.bFilterOnlyOpen = !!NewCheckedState;
+		TBSettings.bFilterOnlyOpen = NewCheckedState == ECheckBoxState::Checked;
 		TBSettings.SaveSettings();
 	}
 }
 
-void STaskBrowser::OnAssignedToMeChanged( const ESlateCheckBoxState::Type NewCheckedState )
+void STaskBrowser::OnAssignedToMeChanged( const ECheckBoxState NewCheckedState )
 {
 	// Refresh the GUI.  It will apply the updated filters.
 	RefreshGUI( ETaskBrowserGUIRefreshOptions::RebuildTaskList );
@@ -882,12 +882,12 @@ void STaskBrowser::OnAssignedToMeChanged( const ESlateCheckBoxState::Type NewChe
 	{
 		FTaskBrowserSettings TBSettings;
 		TBSettings.LoadSettings();
-		TBSettings.bFilterAssignedToMe = !!NewCheckedState;
+		TBSettings.bFilterAssignedToMe = NewCheckedState == ECheckBoxState::Checked;
 		TBSettings.SaveSettings();
 	}
 }
 
-void STaskBrowser::OnCreatedByMeChanged( const ESlateCheckBoxState::Type NewCheckedState )
+void STaskBrowser::OnCreatedByMeChanged( const ECheckBoxState NewCheckedState )
 {
 	// Refresh the GUI.  It will apply the updated filters.
 	RefreshGUI( ETaskBrowserGUIRefreshOptions::RebuildTaskList );
@@ -896,12 +896,12 @@ void STaskBrowser::OnCreatedByMeChanged( const ESlateCheckBoxState::Type NewChec
 	{
 		FTaskBrowserSettings TBSettings;
 		TBSettings.LoadSettings();
-		TBSettings.bFilterCreatedByMe = !!NewCheckedState;
+		TBSettings.bFilterCreatedByMe = NewCheckedState == ECheckBoxState::Checked;
 		TBSettings.SaveSettings();
 	}
 }
 
-void STaskBrowser::OnCurrentMapChanged( const ESlateCheckBoxState::Type NewCheckedState )
+void STaskBrowser::OnCurrentMapChanged( const ECheckBoxState NewCheckedState )
 {
 	// Refresh the GUI.  It will apply the updated filters.
 	RefreshGUI( ETaskBrowserGUIRefreshOptions::RebuildTaskList );
@@ -910,7 +910,7 @@ void STaskBrowser::OnCurrentMapChanged( const ESlateCheckBoxState::Type NewCheck
 	{
 		FTaskBrowserSettings TBSettings;
 		TBSettings.LoadSettings();
-		TBSettings.bFilterCurrentMap = !!NewCheckedState;
+		TBSettings.bFilterCurrentMap = NewCheckedState == ECheckBoxState::Checked;
 		TBSettings.SaveSettings();
 	}
 }
diff --git a/Engine/Source/Editor/TaskBrowser/Private/STaskBrowser.h b/Engine/Source/Editor/TaskBrowser/Private/STaskBrowser.h
--- a/Engine/Source/Editor/TaskBrowser/Private/STaskBrowser.h
+++ b/Engine/Source/Editor/TaskBrowser/Private/STaskBrowser.h
@@ -169,16 +169,16 @@ private:
 	void OnDatabaseFilterSelect( TSharedPtr<FString> InFilter, ESelectInfo::Type SelectInfo );
 
 	/** Handle when the Open Only flag is toggled */
-	void OnOpenOnlyChanged( const ESlateCheckBoxState::Type NewCheckedState );
+	void OnOpenOnlyChanged( const ECheckBoxState NewCheckedState );
 
 	/** Handle when the Assigned To Me flag is toggled */
-	void OnAssignedToMeChanged( const ESlateCheckBoxState::Type NewCheckedState );
+	void OnAssignedToMeChanged( const ECheckBoxState NewCheckedState );
 
 	/** Handle when the Created By Me flag is toggled */
-	void OnCreatedByMeChanged( const ESlateCheckBoxState::Type NewCheckedState );
+	void OnCreatedByMeChanged( const ECheckBoxState NewCheckedState );
 
 	/** Handle when the Current Map flag is toggled */
-	void OnCurrentMapChanged( const ESlateCheckBoxState::Type NewCheckedState );
+	void OnCurrentMapChanged( const ECheckBoxState NewCheckedState );
 
 	/** Create a widget to display a row of our table correctly */
 	TSharedRef<ITableRow> OnTaskGenerateRow( TSharedPtr<FTaskOverview> InTaskOverview, const TSharedRef<STableViewBase>& OwnerTable );
diff --git a/Engine/Source/Editor/TaskBrowser/Private/STaskSettings.cpp b/Engine/Source/Editor/TaskBrowser/Private/STaskSettings.cpp
--- a/Engine/Source/Editor/TaskBrowser/Private/STaskSettings.cpp
+++ b/Engine/Source/Editor/TaskBrowser/Private/STaskSettings.cpp
@@ -168,7 +168,7 @@ void STaskSettings::Construct(const FArguments& InArgs)
 							.AutoWidth()
 							[
 								SAssignNew(Autoconnect, SCheckBox)
-								.IsChecked( TBSettings.bAutoConnectAtStartup ? ESlateCheckBoxState::Checked : ESlateCheckBoxState::Unchecked )
+								.IsChecked( TBSettings.bAutoConnectAtStartup ? ECheckBoxState::Checked : ECheckBoxState::Unchecked )
 							]
 							+SHorizontalBox::Slot()
 							.AutoWidth()
@@ -183,7 +183,7 @@ void STaskSettings::Construct(const FArguments& InArgs)
 							.AutoWidth()
 							[
 								SAssignNew(SingleSignOn, SCheckBox)
-								.IsChecked( TBSettings.bUseSingleSignOn ? ESlateCheckBoxState::Checked : ESlateCheckBoxState::Unchecked )
+								.IsChecked( TBSettings.bUseSingleSignOn ? ECheckBoxState::Checked : ECheckBoxState::Unchecked )
 								.OnCheckStateChanged( this, &STaskSettings::OnSingleSignOnChanged )
 							]
 							+SHorizontalBox::Slot()
@@ -234,10 +234,10 @@ void STaskSettings::Construct(const FArguments& InArgs)
 	];
 }
 
-void STaskSettings::OnSingleSignOnChanged( const ESlateCheckBoxState::Type NewCheckedState )
+void STaskSettings::OnSingleSignOnChanged( const ECheckBoxState NewCheckedState )
 {
 	// Enable/Disable the login and password fields
-	const bool bEnabled = ( NewCheckedState == ESlateCheckBoxState::Checked ? false : true );
+	const bool bEnabled = ( NewCheckedState == ECheckBoxState::Checked ? false : true );
 	Login->SetEnabled( bEnabled );
 	Password->SetEnabled( bEnabled );
 }
diff --git a/Engine/Source/Editor/TaskBrowser/Private/STaskSettings.h b/Engine/Source/Editor/TaskBrowser/Private/STaskSettings.h
--- a/Engine/Source/Editor/TaskBrowser/Private/STaskSettings.h
+++ b/Engine/Source/Editor/TaskBrowser/Private/STaskSettings.h
@@ -24,7 +24,7 @@ public:
 private:
 
 	/** Handle when the Single Sign On flag is checked */
-	void OnSingleSignOnChanged( const ESlateCheckBoxState::Type NewCheckedState );
+	void OnSingleSignOnChanged( const ECheckBoxState NewCheckedState );
 
 	/** Handle when the OK button is pressed */
 	FReply OnOKClicked();
diff --git a/Engine/Source/Editor/TextureEditor/Private/TextureEditorToolkit.cpp b/Engine/Source/Editor/TextureEditor/Private/TextureEditorToolkit.cpp
--- a/Engine/Source/Editor/TextureEditor/Private/TextureEditorToolkit.cpp
+++ b/Engine/Source/Editor/TextureEditor/Private/TextureEditorToolkit.cpp
@@ -910,15 +910,15 @@ bool FTextureEditorToolkit::HandleGreenChannelActionIsChecked( ) const
 }
 
 
-void FTextureEditorToolkit::HandleMipLevelCheckBoxCheckedStateChanged( ESlateCheckBoxState::Type InNewState )
+void FTextureEditorToolkit::HandleMipLevelCheckBoxCheckedStateChanged( ECheckBoxState InNewState )
 {
-	bUseSpecifiedMipLevel = InNewState == ESlateCheckBoxState::Checked;
+	bUseSpecifiedMipLevel = InNewState == ECheckBoxState::Checked;
 }
 
 
-ESlateCheckBoxState::Type FTextureEditorToolkit::HandleMipLevelCheckBoxIsChecked( ) const
+ECheckBoxState FTextureEditorToolkit::HandleMipLevelCheckBoxIsChecked( ) const
 {
-	return GetUseSpecifiedMip() ? ESlateCheckBoxState::Checked : ESlateCheckBoxState::Unchecked;
+	return GetUseSpecifiedMip() ? ECheckBoxState::Checked : ECheckBoxState::Unchecked;
 }
 
 
diff --git a/Engine/Source/Editor/TextureEditor/Private/TextureEditorToolkit.h b/Engine/Source/Editor/TextureEditor/Private/TextureEditorToolkit.h
--- a/Engine/Source/Editor/TextureEditor/Private/TextureEditorToolkit.h
+++ b/Engine/Source/Editor/TextureEditor/Private/TextureEditorToolkit.h
@@ -166,10 +166,10 @@ private:
 	bool HandleGreenChannelActionIsChecked( ) const;
 
 	// Callback for changing the checked state of the MipMap check box.
-	void HandleMipLevelCheckBoxCheckedStateChanged( ESlateCheckBoxState::Type InNewState );
+	void HandleMipLevelCheckBoxCheckedStateChanged( ECheckBoxState InNewState );
 
 	// Callback for getting the checked state of the MipMap check box.
-	ESlateCheckBoxState::Type HandleMipLevelCheckBoxIsChecked( ) const;
+	ECheckBoxState HandleMipLevelCheckBoxIsChecked( ) const;
 
 	// Callback for determining whether the MipMap check box is enabled.
 	bool HandleMipLevelCheckBoxIsEnabled( ) const;
diff --git a/Engine/Source/Editor/TranslationEditor/Private/TranslationPickerWidget.cpp b/Engine/Source/Editor/TranslationEditor/Private/TranslationPickerWidget.cpp
--- a/Engine/Source/Editor/TranslationEditor/Private/TranslationPickerWidget.cpp
+++ b/Engine/Source/Editor/TranslationEditor/Private/TranslationPickerWidget.cpp
@@ -129,12 +129,12 @@ bool STranslationWidgetPicker::IsPicking() const
 	return false;
 }
 
-ESlateCheckBoxState::Type STranslationWidgetPicker::OnIsChecked() const
+ECheckBoxState STranslationWidgetPicker::OnIsChecked() const
 {
-	return IsPicking() ? ESlateCheckBoxState::Checked : ESlateCheckBoxState::Unchecked;
+	return IsPicking() ? ECheckBoxState::Checked : ECheckBoxState::Unchecked;
 }
 
-void STranslationWidgetPicker::OnCheckStateChanged(const ESlateCheckBoxState::Type NewCheckedState)
+void STranslationWidgetPicker::OnCheckStateChanged(const ECheckBoxState NewCheckedState)
 {
 	OnClicked();
 }
diff --git a/Engine/Source/Editor/TranslationEditor/Private/TranslationPickerWidget.h b/Engine/Source/Editor/TranslationEditor/Private/TranslationPickerWidget.h
--- a/Engine/Source/Editor/TranslationEditor/Private/TranslationPickerWidget.h
+++ b/Engine/Source/Editor/TranslationEditor/Private/TranslationPickerWidget.h
@@ -32,14 +32,14 @@ private:
 	/**
 	* Called by slate to determine if this button should appear checked
 	*
-	* @return ESlateCheckBoxState::Checked if it should be checked, ESlateCheckBoxState::Unchecked if not.
+	* @return ECheckBoxState::Checked if it should be checked, ECheckBoxState::Unchecked if not.
 	*/
-	ESlateCheckBoxState::Type OnIsChecked() const;
+	ECheckBoxState OnIsChecked() const;
 
 	/**
 	* Called by Slate when this tool bar check box button is toggled
 	*/
-	void OnCheckStateChanged(const ESlateCheckBoxState::Type NewCheckedState);
+	void OnCheckStateChanged(const ECheckBoxState NewCheckedState);
 };
 
 #undef LOCTEXT_NAMESPACE
\ No newline at end of file
diff --git a/Engine/Source/Editor/UMGEditor/Private/Customizations/HorizontalAlignmentCustomization.cpp b/Engine/Source/Editor/UMGEditor/Private/Customizations/HorizontalAlignmentCustomization.cpp
--- a/Engine/Source/Editor/UMGEditor/Private/Customizations/HorizontalAlignmentCustomization.cpp
+++ b/Engine/Source/Editor/UMGEditor/Private/Customizations/HorizontalAlignmentCustomization.cpp
@@ -96,20 +96,20 @@ void FHorizontalAlignmentCustomization::CustomizeChildren(TSharedRef<IPropertyHa
 {
 }
 
-void FHorizontalAlignmentCustomization::HandleCheckStateChanged(ESlateCheckBoxState::Type InCheckboxState, TSharedRef<IPropertyHandle> PropertyHandle, EHorizontalAlignment ToAlignment)
+void FHorizontalAlignmentCustomization::HandleCheckStateChanged(ECheckBoxState InCheckboxState, TSharedRef<IPropertyHandle> PropertyHandle, EHorizontalAlignment ToAlignment)
 {
 	PropertyHandle->SetValue((uint8)ToAlignment);
 }
 
-ESlateCheckBoxState::Type FHorizontalAlignmentCustomization::GetCheckState(TSharedRef<IPropertyHandle> PropertyHandle, EHorizontalAlignment ForAlignment) const
+ECheckBoxState FHorizontalAlignmentCustomization::GetCheckState(TSharedRef<IPropertyHandle> PropertyHandle, EHorizontalAlignment ForAlignment) const
 {
 	uint8 Value;
 	if ( PropertyHandle->GetValue(Value) )
 	{
-		return Value == ForAlignment ? ESlateCheckBoxState::Checked : ESlateCheckBoxState::Unchecked;
+		return Value == ForAlignment ? ECheckBoxState::Checked : ECheckBoxState::Unchecked;
 	}
 
-	return ESlateCheckBoxState::Unchecked;
+	return ECheckBoxState::Unchecked;
 }
 
 #undef LOCTEXT_NAMESPACE
diff --git a/Engine/Source/Editor/UMGEditor/Private/Customizations/HorizontalAlignmentCustomization.h b/Engine/Source/Editor/UMGEditor/Private/Customizations/HorizontalAlignmentCustomization.h
--- a/Engine/Source/Editor/UMGEditor/Private/Customizations/HorizontalAlignmentCustomization.h
+++ b/Engine/Source/Editor/UMGEditor/Private/Customizations/HorizontalAlignmentCustomization.h
@@ -20,7 +20,7 @@ public:
 
 	virtual void CustomizeChildren(TSharedRef<IPropertyHandle> PropertyHandle, class IDetailChildrenBuilder& ChildBuilder, IPropertyTypeCustomizationUtils& CustomizationUtils) override;
 	
-	void HandleCheckStateChanged(ESlateCheckBoxState::Type InCheckboxState, TSharedRef<IPropertyHandle> PropertyHandle, EHorizontalAlignment ToAlignment);
+	void HandleCheckStateChanged(ECheckBoxState InCheckboxState, TSharedRef<IPropertyHandle> PropertyHandle, EHorizontalAlignment ToAlignment);
 
-	ESlateCheckBoxState::Type GetCheckState(TSharedRef<IPropertyHandle> PropertyHandle, EHorizontalAlignment ForAlignment) const;
+	ECheckBoxState GetCheckState(TSharedRef<IPropertyHandle> PropertyHandle, EHorizontalAlignment ForAlignment) const;
 };
diff --git a/Engine/Source/Editor/UMGEditor/Private/Customizations/SlateChildSizeCustomization.cpp b/Engine/Source/Editor/UMGEditor/Private/Customizations/SlateChildSizeCustomization.cpp
--- a/Engine/Source/Editor/UMGEditor/Private/Customizations/SlateChildSizeCustomization.cpp
+++ b/Engine/Source/Editor/UMGEditor/Private/Customizations/SlateChildSizeCustomization.cpp
@@ -103,20 +103,20 @@ void FSlateChildSizeCustomization::CustomizeChildren(TSharedRef<IPropertyHandle>
 	
 }
 
-void FSlateChildSizeCustomization::HandleCheckStateChanged(ESlateCheckBoxState::Type InCheckboxState, TSharedPtr<IPropertyHandle> PropertyHandle, ESlateSizeRule::Type ToRule)
+void FSlateChildSizeCustomization::HandleCheckStateChanged(ECheckBoxState InCheckboxState, TSharedPtr<IPropertyHandle> PropertyHandle, ESlateSizeRule::Type ToRule)
 {
 	PropertyHandle->SetValue((uint8)ToRule);
 }
 
-ESlateCheckBoxState::Type FSlateChildSizeCustomization::GetCheckState(TSharedPtr<IPropertyHandle> PropertyHandle, ESlateSizeRule::Type ForRule) const
+ECheckBoxState FSlateChildSizeCustomization::GetCheckState(TSharedPtr<IPropertyHandle> PropertyHandle, ESlateSizeRule::Type ForRule) const
 {
 	uint8 Value;
 	if ( PropertyHandle->GetValue(Value) )
 	{
-		return Value == ForRule ? ESlateCheckBoxState::Checked : ESlateCheckBoxState::Unchecked;
+		return Value == ForRule ? ECheckBoxState::Checked : ECheckBoxState::Unchecked;
 	}
 
-	return ESlateCheckBoxState::Unchecked;
+	return ECheckBoxState::Unchecked;
 }
 
 TOptional<float> FSlateChildSizeCustomization::GetValue(TSharedPtr<IPropertyHandle> ValueHandle) const
diff --git a/Engine/Source/Editor/UMGEditor/Private/Customizations/SlateChildSizeCustomization.h b/Engine/Source/Editor/UMGEditor/Private/Customizations/SlateChildSizeCustomization.h
--- a/Engine/Source/Editor/UMGEditor/Private/Customizations/SlateChildSizeCustomization.h
+++ b/Engine/Source/Editor/UMGEditor/Private/Customizations/SlateChildSizeCustomization.h
@@ -22,9 +22,9 @@ public:
 
 	virtual void CustomizeChildren(TSharedRef<IPropertyHandle> PropertyHandle, class IDetailChildrenBuilder& ChildBuilder, IPropertyTypeCustomizationUtils& CustomizationUtils) override;
 
-	void HandleCheckStateChanged(ESlateCheckBoxState::Type InCheckboxState, TSharedPtr<IPropertyHandle> PropertyHandle, ESlateSizeRule::Type ToRule);
+	void HandleCheckStateChanged(ECheckBoxState InCheckboxState, TSharedPtr<IPropertyHandle> PropertyHandle, ESlateSizeRule::Type ToRule);
 
-	ESlateCheckBoxState::Type GetCheckState(TSharedPtr<IPropertyHandle> PropertyHandle, ESlateSizeRule::Type ForRule) const;
+	ECheckBoxState GetCheckState(TSharedPtr<IPropertyHandle> PropertyHandle, ESlateSizeRule::Type ForRule) const;
 
 	TOptional<float> GetValue(TSharedPtr<IPropertyHandle> ValueHandle) const;
 
diff --git a/Engine/Source/Editor/UMGEditor/Private/Customizations/TextJustifyCustomization.cpp b/Engine/Source/Editor/UMGEditor/Private/Customizations/TextJustifyCustomization.cpp
--- a/Engine/Source/Editor/UMGEditor/Private/Customizations/TextJustifyCustomization.cpp
+++ b/Engine/Source/Editor/UMGEditor/Private/Customizations/TextJustifyCustomization.cpp
@@ -80,20 +80,20 @@ void FTextJustifyCustomization::CustomizeChildren(TSharedRef<IPropertyHandle> Pr
 {
 }
 
-void FTextJustifyCustomization::HandleCheckStateChanged(ESlateCheckBoxState::Type InCheckboxState, TSharedRef<IPropertyHandle> PropertyHandle, ETextJustify::Type ToAlignment)
+void FTextJustifyCustomization::HandleCheckStateChanged(ECheckBoxState InCheckboxState, TSharedRef<IPropertyHandle> PropertyHandle, ETextJustify::Type ToAlignment)
 {
 	PropertyHandle->SetValue((uint8)ToAlignment);
 }
 
-ESlateCheckBoxState::Type FTextJustifyCustomization::GetCheckState(TSharedRef<IPropertyHandle> PropertyHandle, ETextJustify::Type ForAlignment) const
+ECheckBoxState FTextJustifyCustomization::GetCheckState(TSharedRef<IPropertyHandle> PropertyHandle, ETextJustify::Type ForAlignment) const
 {
 	uint8 Value;
 	if ( PropertyHandle->GetValue(Value) )
 	{
-		return Value == ForAlignment ? ESlateCheckBoxState::Checked : ESlateCheckBoxState::Unchecked;
+		return Value == ForAlignment ? ECheckBoxState::Checked : ECheckBoxState::Unchecked;
 	}
 
-	return ESlateCheckBoxState::Unchecked;
+	return ECheckBoxState::Unchecked;
 }
 
 #undef LOCTEXT_NAMESPACE
diff --git a/Engine/Source/Editor/UMGEditor/Private/Customizations/TextJustifyCustomization.h b/Engine/Source/Editor/UMGEditor/Private/Customizations/TextJustifyCustomization.h
--- a/Engine/Source/Editor/UMGEditor/Private/Customizations/TextJustifyCustomization.h
+++ b/Engine/Source/Editor/UMGEditor/Private/Customizations/TextJustifyCustomization.h
@@ -19,7 +19,7 @@ public:
 
 	virtual void CustomizeChildren(TSharedRef<IPropertyHandle> PropertyHandle, class IDetailChildrenBuilder& ChildBuilder, IPropertyTypeCustomizationUtils& CustomizationUtils) override;
 	
-	void HandleCheckStateChanged(ESlateCheckBoxState::Type InCheckboxState, TSharedRef<IPropertyHandle> PropertyHandle, ETextJustify::Type ToAlignment);
+	void HandleCheckStateChanged(ECheckBoxState InCheckboxState, TSharedRef<IPropertyHandle> PropertyHandle, ETextJustify::Type ToAlignment);
 
-	ESlateCheckBoxState::Type GetCheckState(TSharedRef<IPropertyHandle> PropertyHandle, ETextJustify::Type ForAlignment) const;
+	ECheckBoxState GetCheckState(TSharedRef<IPropertyHandle> PropertyHandle, ETextJustify::Type ForAlignment) const;
 };
diff --git a/Engine/Source/Editor/UMGEditor/Private/Customizations/VerticalAlignmentCustomization.cpp b/Engine/Source/Editor/UMGEditor/Private/Customizations/VerticalAlignmentCustomization.cpp
--- a/Engine/Source/Editor/UMGEditor/Private/Customizations/VerticalAlignmentCustomization.cpp
+++ b/Engine/Source/Editor/UMGEditor/Private/Customizations/VerticalAlignmentCustomization.cpp
@@ -96,20 +96,20 @@ void FVerticalAlignmentCustomization::CustomizeChildren(TSharedRef<IPropertyHand
 {
 }
 
-void FVerticalAlignmentCustomization::HandleCheckStateChanged(ESlateCheckBoxState::Type InCheckboxState, TSharedRef<IPropertyHandle> PropertyHandle, EVerticalAlignment ToAlignment)
+void FVerticalAlignmentCustomization::HandleCheckStateChanged(ECheckBoxState InCheckboxState, TSharedRef<IPropertyHandle> PropertyHandle, EVerticalAlignment ToAlignment)
 {
 	PropertyHandle->SetValue((uint8)ToAlignment);
 }
 
-ESlateCheckBoxState::Type FVerticalAlignmentCustomization::GetCheckState(TSharedRef<IPropertyHandle> PropertyHandle, EVerticalAlignment ForAlignment) const
+ECheckBoxState FVerticalAlignmentCustomization::GetCheckState(TSharedRef<IPropertyHandle> PropertyHandle, EVerticalAlignment ForAlignment) const
 {
 	uint8 Value;
 	if ( PropertyHandle->GetValue(Value) )
 	{
-		return Value == ForAlignment ? ESlateCheckBoxState::Checked : ESlateCheckBoxState::Unchecked;
+		return Value == ForAlignment ? ECheckBoxState::Checked : ECheckBoxState::Unchecked;
 	}
 
-	return ESlateCheckBoxState::Unchecked;
+	return ECheckBoxState::Unchecked;
 }
 
 #undef LOCTEXT_NAMESPACE
diff --git a/Engine/Source/Editor/UMGEditor/Private/Customizations/VerticalAlignmentCustomization.h b/Engine/Source/Editor/UMGEditor/Private/Customizations/VerticalAlignmentCustomization.h
--- a/Engine/Source/Editor/UMGEditor/Private/Customizations/VerticalAlignmentCustomization.h
+++ b/Engine/Source/Editor/UMGEditor/Private/Customizations/VerticalAlignmentCustomization.h
@@ -19,7 +19,7 @@ public:
 
 	virtual void CustomizeChildren(TSharedRef<IPropertyHandle> PropertyHandle, class IDetailChildrenBuilder& ChildBuilder, IPropertyTypeCustomizationUtils& CustomizationUtils) override;
 	
-	void HandleCheckStateChanged(ESlateCheckBoxState::Type InCheckboxState, TSharedRef<IPropertyHandle> PropertyHandle, EVerticalAlignment ToAlignment);
+	void HandleCheckStateChanged(ECheckBoxState InCheckboxState, TSharedRef<IPropertyHandle> PropertyHandle, EVerticalAlignment ToAlignment);
 
-	ESlateCheckBoxState::Type GetCheckState(TSharedRef<IPropertyHandle> PropertyHandle, EVerticalAlignment ForAlignment) const;
+	ECheckBoxState GetCheckState(TSharedRef<IPropertyHandle> PropertyHandle, EVerticalAlignment ForAlignment) const;
 };
diff --git a/Engine/Source/Editor/UMGEditor/Private/Designer/SDesignerToolBar.cpp b/Engine/Source/Editor/UMGEditor/Private/Designer/SDesignerToolBar.cpp
--- a/Engine/Source/Editor/UMGEditor/Private/Designer/SDesignerToolBar.cpp
+++ b/Engine/Source/Editor/UMGEditor/Private/Designer/SDesignerToolBar.cpp
@@ -68,12 +68,12 @@ TSharedRef< SWidget > SDesignerToolBar::MakeToolBar(const TSharedPtr< FExtender
 	return ToolbarBuilder.MakeWidget();
 }
 
-ESlateCheckBoxState::Type SDesignerToolBar::IsLocationGridSnapChecked() const
+ECheckBoxState SDesignerToolBar::IsLocationGridSnapChecked() const
 {
-	return GetDefault<UWidgetDesignerSettings>()->GridSnapEnabled ? ESlateCheckBoxState::Checked : ESlateCheckBoxState::Unchecked;
+	return GetDefault<UWidgetDesignerSettings>()->GridSnapEnabled ? ECheckBoxState::Checked : ECheckBoxState::Unchecked;
 }
 
-void SDesignerToolBar::HandleToggleLocationGridSnap(ESlateCheckBoxState::Type InState)
+void SDesignerToolBar::HandleToggleLocationGridSnap(ECheckBoxState InState)
 {
 	UWidgetDesignerSettings* ViewportSettings = GetMutableDefault<UWidgetDesignerSettings>();
 	ViewportSettings->GridSnapEnabled = !ViewportSettings->GridSnapEnabled;
diff --git a/Engine/Source/Editor/UMGEditor/Private/Designer/SDesignerToolBar.h b/Engine/Source/Editor/UMGEditor/Private/Designer/SDesignerToolBar.h
--- a/Engine/Source/Editor/UMGEditor/Private/Designer/SDesignerToolBar.h
+++ b/Engine/Source/Editor/UMGEditor/Private/Designer/SDesignerToolBar.h
@@ -26,8 +26,8 @@ public:
 
 private:
 	// Begin Grid Snapping
-	ESlateCheckBoxState::Type IsLocationGridSnapChecked() const;
-	void HandleToggleLocationGridSnap(ESlateCheckBoxState::Type InState);
+	ECheckBoxState IsLocationGridSnapChecked() const;
+	void HandleToggleLocationGridSnap(ECheckBoxState InState);
 	FText GetLocationGridLabel() const;
 	TSharedRef<SWidget> FillLocationGridSnapMenu();
 	TSharedRef<SWidget> BuildLocationGridCheckBoxList(FName InExtentionHook, const FText& InHeading, const TArray<int32>& InGridSizes) const;
diff --git a/Engine/Source/Editor/UMGEditor/Private/Details/SWidgetDetailsView.cpp b/Engine/Source/Editor/UMGEditor/Private/Details/SWidgetDetailsView.cpp
--- a/Engine/Source/Editor/UMGEditor/Private/Details/SWidgetDetailsView.cpp
+++ b/Engine/Source/Editor/UMGEditor/Private/Details/SWidgetDetailsView.cpp
@@ -430,21 +430,21 @@ void SWidgetDetailsView::HandleNameTextCommitted(const FText& Text, ETextCommit:
 	}
 }
 
-ESlateCheckBoxState::Type SWidgetDetailsView::GetIsVariable() const
+ECheckBoxState SWidgetDetailsView::GetIsVariable() const
 {
 	if ( SelectedObjects.Num() == 1 )
 	{
 		UWidget* Widget = Cast<UWidget>(SelectedObjects[0].Get());
 		if ( Widget )
 		{
-			return Widget->bIsVariable ? ESlateCheckBoxState::Checked : ESlateCheckBoxState::Unchecked;
+			return Widget->bIsVariable ? ECheckBoxState::Checked : ECheckBoxState::Unchecked;
 		}
 	}
 
-	return ESlateCheckBoxState::Unchecked;
+	return ECheckBoxState::Unchecked;
 }
 
-void SWidgetDetailsView::HandleIsVariableChanged(ESlateCheckBoxState::Type CheckState)
+void SWidgetDetailsView::HandleIsVariableChanged(ECheckBoxState CheckState)
 {
 	if ( SelectedObjects.Num() == 1 )
 	{
@@ -463,7 +463,7 @@ void SWidgetDetailsView::HandleIsVariableChanged(ESlateCheckBoxState::Type Check
 			Template->Modify();
 			Preview->Modify();
 
-			Template->bIsVariable = Preview->bIsVariable = CheckState == ESlateCheckBoxState::Checked ? true : false;
+			Template->bIsVariable = Preview->bIsVariable = CheckState == ECheckBoxState::Checked ? true : false;
 
 			// Refresh references and flush editors
 			FBlueprintEditorUtils::MarkBlueprintAsStructurallyModified(Blueprint);
diff --git a/Engine/Source/Editor/UMGEditor/Private/Details/SWidgetDetailsView.h b/Engine/Source/Editor/UMGEditor/Private/Details/SWidgetDetailsView.h
--- a/Engine/Source/Editor/UMGEditor/Private/Details/SWidgetDetailsView.h
+++ b/Engine/Source/Editor/UMGEditor/Private/Details/SWidgetDetailsView.h
@@ -52,9 +52,9 @@ private:
 	void HandleNameTextChanged(const FText& Text);
 	void HandleNameTextCommitted(const FText& Text, ETextCommit::Type CommitType);
 
-	ESlateCheckBoxState::Type GetIsVariable() const;
+	ECheckBoxState GetIsVariable() const;
 
-	void HandleIsVariableChanged(ESlateCheckBoxState::Type CheckState);
+	void HandleIsVariableChanged(ECheckBoxState CheckState);
 
 	EVisibility GetCategoryAreaVisibility() const;
 	FText GetCategoryText() const;
diff --git a/Engine/Source/Editor/UnrealEd/Classes/Niagara/NiagaraNodeOp.h b/Engine/Source/Editor/UnrealEd/Classes/Niagara/NiagaraNodeOp.h
--- a/Engine/Source/Editor/UnrealEd/Classes/Niagara/NiagaraNodeOp.h
+++ b/Engine/Source/Editor/UnrealEd/Classes/Niagara/NiagaraNodeOp.h
@@ -63,10 +63,10 @@ protected:
 	//EVisibility ExpressionPreviewVisibility() const;
 
 	/** Show/hide Expression Preview */
-	//void OnExpressionPreviewChanged(const ESlateCheckBoxState::Type NewCheckedState);
+	//void OnExpressionPreviewChanged(const ECheckBoxState NewCheckedState);
 
 	/** hidden == unchecked, shown == checked */
-	//ESlateCheckBoxState::Type IsExpressionPreviewChecked() const;
+	//ECheckBoxState IsExpressionPreviewChecked() const;
 
 	/** Up when shown, down when hidden */
 	//const FSlateBrush* GetExpressionPreviewArrow() const;
diff --git a/Engine/Source/Editor/UnrealEd/Private/ConsolidateWindow.cpp b/Engine/Source/Editor/UnrealEd/Private/ConsolidateWindow.cpp
--- a/Engine/Source/Editor/UnrealEd/Private/ConsolidateWindow.cpp
+++ b/Engine/Source/Editor/UnrealEd/Private/ConsolidateWindow.cpp
@@ -53,7 +53,7 @@ public:
 		/** 
 		 * Callback function used to tell the ListBox parent what item has been selected
 		 */
-		void OnAssetSelected( ESlateCheckBoxState::Type NewCheckedState )
+		void OnAssetSelected( ECheckBoxState NewCheckedState )
 		{
 			check(Parent != NULL);
 			Parent->SetSelectedListItem(this);
@@ -62,9 +62,9 @@ public:
 		/**
 		 * Callback function used to ensure only one item is highlighted (selected) at a time
 		 */
-		ESlateCheckBoxState::Type IsAssetSelected() const
+		ECheckBoxState IsAssetSelected() const
 		{
-			return (Parent->GetSelectedListItem() == this) ? ESlateCheckBoxState::Checked : ESlateCheckBoxState::Unchecked;
+			return (Parent->GetSelectedListItem() == this) ? ECheckBoxState::Checked : ECheckBoxState::Unchecked;
 		}
 
 		/** @return		The full name of the object this item represents */
@@ -249,10 +249,10 @@ private:
 	bool IsConsolidateButtonEnabled() const;
 
 	/** @return		true if the user has elected to "Save Dirty Packages" */
-	ESlateCheckBoxState::Type IsSavePackagesChecked() const;
+	ECheckBoxState IsSavePackagesChecked() const;
 	
 	/** Callback	called when the user checks/unchecks the "Save Dirty Packages" button */
-	void OnSavePackagesCheckStateChanged(ESlateCheckBoxState::Type NewCheckedState);
+	void OnSavePackagesCheckStateChanged(ECheckBoxState NewCheckedState);
 
 	/** Callback	Called to determin the error message visibility */
 	EVisibility IsErrorPanelVisible() const;
@@ -452,14 +452,14 @@ bool SConsolidateToolWidget::IsConsolidateButtonEnabled() const
 	return (ConsolidationObjects.Num() > 1 && SelectedListItem != NULL);
 }
 
-ESlateCheckBoxState::Type SConsolidateToolWidget::IsSavePackagesChecked() const
+ECheckBoxState SConsolidateToolWidget::IsSavePackagesChecked() const
 {
-	return bSavePackagesChecked ? ESlateCheckBoxState::Checked : ESlateCheckBoxState::Unchecked;
+	return bSavePackagesChecked ? ECheckBoxState::Checked : ECheckBoxState::Unchecked;
 }
 
-void SConsolidateToolWidget::OnSavePackagesCheckStateChanged(ESlateCheckBoxState::Type NewCheckedState)
+void SConsolidateToolWidget::OnSavePackagesCheckStateChanged(ECheckBoxState NewCheckedState)
 {
-	bSavePackagesChecked = (NewCheckedState == ESlateCheckBoxState::Checked);
+	bSavePackagesChecked = (NewCheckedState == ECheckBoxState::Checked);
 }
 
 EVisibility SConsolidateToolWidget::IsErrorPanelVisible() const
diff --git a/Engine/Source/Editor/UnrealEd/Private/Dialogs/Dialogs.cpp b/Engine/Source/Editor/UnrealEd/Private/Dialogs/Dialogs.cpp
--- a/Engine/Source/Editor/UnrealEd/Private/Dialogs/Dialogs.cpp
+++ b/Engine/Source/Editor/UnrealEd/Private/Dialogs/Dialogs.cpp
@@ -521,7 +521,7 @@ public:
 		.AutoWidth()
 		[
 			SNew(SCheckBox)
-			.IsChecked(InArgs._bDefaultCheckValue ? ESlateCheckBoxState::Checked : ESlateCheckBoxState::Unchecked)
+			.IsChecked(InArgs._bDefaultCheckValue ? ECheckBoxState::Checked : ECheckBoxState::Unchecked)
 			.OnCheckStateChanged(this, &SModalDialogWithCheckbox::OnCheckboxClicked)
 			.Visibility(this, &SModalDialogWithCheckbox::GetCheckboxVisibility)
 			[
@@ -622,9 +622,9 @@ private:
 	}
 
 	/** Used as a delegate for the Checkboxs OnClicked method */
-	void OnCheckboxClicked(ESlateCheckBoxState::Type InNewState)
+	void OnCheckboxClicked(ECheckBoxState InNewState)
 	{
-		bCheckboxResult = InNewState == ESlateCheckBoxState::Checked;
+		bCheckboxResult = InNewState == ECheckBoxState::Checked;
 	}
 
 	/** Used as a delegate for the Checkboxs Visibile method */
diff --git a/Engine/Source/Editor/UnrealEd/Private/Dialogs/SMeshMergingDialog.cpp b/Engine/Source/Editor/UnrealEd/Private/Dialogs/SMeshMergingDialog.cpp
--- a/Engine/Source/Editor/UnrealEd/Private/Dialogs/SMeshMergingDialog.cpp
+++ b/Engine/Source/Editor/UnrealEd/Private/Dialogs/SMeshMergingDialog.cpp
@@ -270,14 +270,14 @@ FReply SMeshMergingDialog::OnMergeClicked()
 	return FReply::Handled();
 }
 
-ESlateCheckBoxState::Type SMeshMergingDialog::GetGenerateLightmapUV() const
+ECheckBoxState SMeshMergingDialog::GetGenerateLightmapUV() const
 {
-	return (MergingSettings.bGenerateLightMapUV ? ESlateCheckBoxState::Checked : ESlateCheckBoxState::Unchecked);
+	return (MergingSettings.bGenerateLightMapUV ? ECheckBoxState::Checked : ECheckBoxState::Unchecked);
 }
 
-void SMeshMergingDialog::SetGenerateLightmapUV(ESlateCheckBoxState::Type NewValue)
+void SMeshMergingDialog::SetGenerateLightmapUV(ECheckBoxState NewValue)
 {
-	MergingSettings.bGenerateLightMapUV = (ESlateCheckBoxState::Checked == NewValue);
+	MergingSettings.bGenerateLightMapUV = (ECheckBoxState::Checked == NewValue);
 }
 
 bool SMeshMergingDialog::IsLightmapChannelEnabled() const
@@ -295,14 +295,14 @@ void SMeshMergingDialog::SetTargetLightMapResolution(TSharedPtr<FString> NewSele
 	TTypeFromString<int32>::FromString(MergingSettings.TargetLightMapResolution, **NewSelection);
 }
 
-ESlateCheckBoxState::Type SMeshMergingDialog::GetImportVertexColors() const
+ECheckBoxState SMeshMergingDialog::GetImportVertexColors() const
 {
-	return (MergingSettings.bImportVertexColors ? ESlateCheckBoxState::Checked : ESlateCheckBoxState::Unchecked);
+	return (MergingSettings.bImportVertexColors ? ECheckBoxState::Checked : ECheckBoxState::Unchecked);
 }
 
-void SMeshMergingDialog::SetImportVertexColors(ESlateCheckBoxState::Type NewValue)
+void SMeshMergingDialog::SetImportVertexColors(ECheckBoxState NewValue)
 {
-	MergingSettings.bImportVertexColors = (ESlateCheckBoxState::Checked == NewValue);
+	MergingSettings.bImportVertexColors = (ECheckBoxState::Checked == NewValue);
 }
 
 void SMeshMergingDialog::OnActorSelectionChanged(const TArray<UObject*>& NewSelection)
@@ -310,24 +310,24 @@ void SMeshMergingDialog::OnActorSelectionChanged(const TArray<UObject*>& NewSele
 	GenerateNewPackageName();
 }
 
-ESlateCheckBoxState::Type SMeshMergingDialog::GetPivotPointAtZero() const
+ECheckBoxState SMeshMergingDialog::GetPivotPointAtZero() const
 {
-	return (MergingSettings.bPivotPointAtZero ? ESlateCheckBoxState::Checked : ESlateCheckBoxState::Unchecked);
+	return (MergingSettings.bPivotPointAtZero ? ECheckBoxState::Checked : ECheckBoxState::Unchecked);
 }
 
-void SMeshMergingDialog::SetPivotPointAtZero(ESlateCheckBoxState::Type NewValue)
+void SMeshMergingDialog::SetPivotPointAtZero(ECheckBoxState NewValue)
 {
-	MergingSettings.bPivotPointAtZero = (ESlateCheckBoxState::Checked == NewValue);
+	MergingSettings.bPivotPointAtZero = (ECheckBoxState::Checked == NewValue);
 }
 
-ESlateCheckBoxState::Type SMeshMergingDialog::GetPlaceInWorld() const
+ECheckBoxState SMeshMergingDialog::GetPlaceInWorld() const
 {
-	return (bPlaceInWorld ? ESlateCheckBoxState::Checked : ESlateCheckBoxState::Unchecked);
+	return (bPlaceInWorld ? ECheckBoxState::Checked : ECheckBoxState::Unchecked);
 }
 
-void SMeshMergingDialog::SetPlaceInWorld(ESlateCheckBoxState::Type NewValue)
+void SMeshMergingDialog::SetPlaceInWorld(ECheckBoxState NewValue)
 {
-	bPlaceInWorld = (ESlateCheckBoxState::Checked == NewValue);
+	bPlaceInWorld = (ECheckBoxState::Checked == NewValue);
 }
 
 void SMeshMergingDialog::GenerateNewPackageName()
diff --git a/Engine/Source/Editor/UnrealEd/Private/Dialogs/SMeshProxyDialog.cpp b/Engine/Source/Editor/UnrealEd/Private/Dialogs/SMeshProxyDialog.cpp
--- a/Engine/Source/Editor/UnrealEd/Private/Dialogs/SMeshProxyDialog.cpp
+++ b/Engine/Source/Editor/UnrealEd/Private/Dialogs/SMeshProxyDialog.cpp
@@ -47,8 +47,8 @@ protected:
 	void ScreenSizeChanged(int32 NewValue);		//used with editable text block (Simplygon)
 
 	/** Recalculate Normals accessors */
-	ESlateCheckBoxState::Type GetRecalculateNormals() const;
-	void SetRecalculateNormals(ESlateCheckBoxState::Type NewValue);
+	ECheckBoxState GetRecalculateNormals() const;
+	void SetRecalculateNormals(ECheckBoxState NewValue);
 
 	/** Hard Angle Threshold accessors */
 	TOptional<float> GetHardAngleThreshold() const;
@@ -60,8 +60,8 @@ protected:
 	void MergeDistanceChanged(int32 NewValue);
 
 	/** Clipping Plane accessors */
-	ESlateCheckBoxState::Type GetUseClippingPlane() const;
-	void SetUseClippingPlane(ESlateCheckBoxState::Type NewValue);
+	ECheckBoxState GetUseClippingPlane() const;
+	void SetUseClippingPlane(ECheckBoxState NewValue);
 	bool UseClippingPlaneEnabled() const;
 	void SetClippingAxis(TSharedPtr<FString> NewSelection, ESelectInfo::Type SelectInfo);
 	TOptional<float> GetClippingLevel() const;
@@ -463,14 +463,14 @@ void SMeshProxyDialog::ScreenSizeChanged(int32 NewValue)
 }
 
 //Recalculate normals
-ESlateCheckBoxState::Type SMeshProxyDialog::GetRecalculateNormals() const
+ECheckBoxState SMeshProxyDialog::GetRecalculateNormals() const
 {
-	return ProxySettings.bRecalculateNormals ? ESlateCheckBoxState::Checked: ESlateCheckBoxState::Unchecked;
+	return ProxySettings.bRecalculateNormals ? ECheckBoxState::Checked: ECheckBoxState::Unchecked;
 }
 
-void SMeshProxyDialog::SetRecalculateNormals(ESlateCheckBoxState::Type NewValue)
+void SMeshProxyDialog::SetRecalculateNormals(ECheckBoxState NewValue)
 {
-	ProxySettings.bRecalculateNormals = (NewValue == ESlateCheckBoxState::Checked);
+	ProxySettings.bRecalculateNormals = (NewValue == ECheckBoxState::Checked);
 }
 
 //Hard Angle Threshold
@@ -506,14 +506,14 @@ void SMeshProxyDialog::MergeDistanceChanged(int32 NewValue)
 }
 
 //Clipping Plane
-ESlateCheckBoxState::Type SMeshProxyDialog::GetUseClippingPlane() const
+ECheckBoxState SMeshProxyDialog::GetUseClippingPlane() const
 {
-	return ProxySettings.bUseClippingPlane ? ESlateCheckBoxState::Checked: ESlateCheckBoxState::Unchecked;
+	return ProxySettings.bUseClippingPlane ? ECheckBoxState::Checked: ECheckBoxState::Unchecked;
 }
 
-void SMeshProxyDialog::SetUseClippingPlane(ESlateCheckBoxState::Type NewValue)
+void SMeshProxyDialog::SetUseClippingPlane(ECheckBoxState NewValue)
 {
-	ProxySettings.bUseClippingPlane = (NewValue == ESlateCheckBoxState::Checked);
+	ProxySettings.bUseClippingPlane = (NewValue == ECheckBoxState::Checked);
 }
 
 bool SMeshProxyDialog::UseClippingPlaneEnabled() const
diff --git a/Engine/Source/Editor/UnrealEd/Private/FileHelpers.cpp b/Engine/Source/Editor/UnrealEd/Private/FileHelpers.cpp
--- a/Engine/Source/Editor/UnrealEd/Private/FileHelpers.cpp
+++ b/Engine/Source/Editor/UnrealEd/Private/FileHelpers.cpp
@@ -1064,12 +1064,12 @@ bool FEditorFileUtils::PromptToCheckoutPackages(bool bCheckDirty, const TArray<U
 			if( SourceControlState.IsValid() && !SourceControlState->IsCurrent() )
 			{				
 				// This package is not at the head revision and it should be ghosted as a result
-				CheckoutPackagesDialogModule.AddPackageItem(CurPackage, CurPackage->GetName(), ESlateCheckBoxState::Unchecked, true, TEXT("SavePackages.SCC_DlgNotCurrent"), SourceControlState->GetDisplayTooltip().ToString());
+				CheckoutPackagesDialogModule.AddPackageItem(CurPackage, CurPackage->GetName(), ECheckBoxState::Unchecked, true, TEXT("SavePackages.SCC_DlgNotCurrent"), SourceControlState->GetDisplayTooltip().ToString());
 			}
 			else if( SourceControlState.IsValid() && SourceControlState->IsCheckedOutOther() )
 			{
 				// This package is checked out by someone else so it should be ghosted
-				CheckoutPackagesDialogModule.AddPackageItem(CurPackage, CurPackage->GetName(), ESlateCheckBoxState::Unchecked, true, TEXT("SavePackages.SCC_DlgCheckedOutOther"), SourceControlState->GetDisplayTooltip().ToString());
+				CheckoutPackagesDialogModule.AddPackageItem(CurPackage, CurPackage->GetName(), ECheckBoxState::Unchecked, true, TEXT("SavePackages.SCC_DlgCheckedOutOther"), SourceControlState->GetDisplayTooltip().ToString());
 			}
 			else
 			{
@@ -1078,7 +1078,7 @@ bool FEditorFileUtils::PromptToCheckoutPackages(bool bCheckDirty, const TArray<U
 				bHavePackageToCheckOut = true;
 				//Add this package to the dialog if its not checked out, in the source control depot, dirty(if we are checking), and read only
 				//This package could also be marked for delete, which we will treat as SCC_ReadOnly until it is time to check it out. At that time, we will revert it.
-				CheckoutPackagesDialogModule.AddPackageItem(CurPackage, CurPackage->GetName(), ESlateCheckBoxState::Checked, false, TEXT("SavePackages.SCC_DlgReadOnly"), Tooltip.ToString());
+				CheckoutPackagesDialogModule.AddPackageItem(CurPackage, CurPackage->GetName(), ECheckBoxState::Checked, false, TEXT("SavePackages.SCC_DlgReadOnly"), Tooltip.ToString());
 			}
 			bPackagesAdded = true;
 		}
@@ -1091,7 +1091,7 @@ bool FEditorFileUtils::PromptToCheckoutPackages(bool bCheckDirty, const TArray<U
 
 			// This package is read only but source control is not available, show the dialog so users can save the package by making the file writable or by connecting to source control.
 			// If we don't care about read-only state, we should allow the user to make the file writable whatever the state of source control.
-			CheckoutPackagesDialogModule.AddPackageItem(CurPackage, CurPackage->GetName(), ESlateCheckBoxState::Unchecked, bIsDisabled, TEXT("SavePackages.SCC_DlgReadOnly"), Tooltip.ToString());
+			CheckoutPackagesDialogModule.AddPackageItem(CurPackage, CurPackage->GetName(), ECheckBoxState::Unchecked, bIsDisabled, TEXT("SavePackages.SCC_DlgReadOnly"), Tooltip.ToString());
 			bPackagesAdded = true;
 		}
 		else if ( OutPackagesNotNeedingCheckout )
@@ -1140,7 +1140,7 @@ bool FEditorFileUtils::PromptToCheckoutPackages(bool bCheckDirty, const TArray<U
 			{
 				// Get the packages that should be checked out from the user's choices in the dialog
 				TArray<UPackage*> PkgsToCheckOut;
-				CheckoutPackagesDialogModule.GetResults( PkgsToCheckOut, ESlateCheckBoxState::Checked );
+				CheckoutPackagesDialogModule.GetResults( PkgsToCheckOut, ECheckBoxState::Checked );
 
 				if(CheckoutPackages(PkgsToCheckOut, OutPackagesCheckedOutOrMadeWritable) == ECommandResult::Cancelled)
 				{
@@ -1156,8 +1156,8 @@ bool FEditorFileUtils::PromptToCheckoutPackages(bool bCheckDirty, const TArray<U
 				// Get the packages that should be made writable out from the user's choices in the dialog
 				TArray<UPackage*> PkgsToMakeWritable;
 				// Both undetermined and checked should be made writable.  Undetermined is only available when packages cant be checked out
-				CheckoutPackagesDialogModule.GetResults( PkgsToMakeWritable, ESlateCheckBoxState::Undetermined );
-				CheckoutPackagesDialogModule.GetResults( PkgsToMakeWritable, ESlateCheckBoxState::Checked);
+				CheckoutPackagesDialogModule.GetResults( PkgsToMakeWritable, ECheckBoxState::Undetermined );
+				CheckoutPackagesDialogModule.GetResults( PkgsToMakeWritable, ECheckBoxState::Checked);
 
 				bool bPackageFailedWritable = false;
 				FString PkgsWhichFailedWritable;
@@ -2741,11 +2741,11 @@ FEditorFileUtils::EPromptReturnCode FEditorFileUtils::PromptForCheckoutAndSave(
 		{
 			for (auto Iter = AddPackageItemsChecked.CreateIterator(); Iter; ++Iter)
 			{
-				PackagesDialogModule.AddPackageItem(*Iter, (*Iter)->GetName(), ESlateCheckBoxState::Checked);
+				PackagesDialogModule.AddPackageItem(*Iter, (*Iter)->GetName(), ECheckBoxState::Checked);
 			}
 			for (auto Iter = AddPackageItemsUnchecked.CreateIterator(); Iter; ++Iter)
 			{
-				PackagesDialogModule.AddPackageItem(*Iter, (*Iter)->GetName(), ESlateCheckBoxState::Unchecked);
+				PackagesDialogModule.AddPackageItem(*Iter, (*Iter)->GetName(), ECheckBoxState::Unchecked);
 			}
 
 			// If valid packages were added to the dialog, display it to the user
@@ -2754,10 +2754,10 @@ FEditorFileUtils::EPromptReturnCode FEditorFileUtils::PromptForCheckoutAndSave(
 			// If the user has responded yes, they want to save the packages they have checked
 			if ( UserResponse == DRT_Save )
 			{
-				PackagesDialogModule.GetResults( FilteredPackages, ESlateCheckBoxState::Checked );
+				PackagesDialogModule.GetResults( FilteredPackages, ECheckBoxState::Checked );
 
 				TArray<UPackage*> UncheckedPackagesRaw;
-				PackagesDialogModule.GetResults( UncheckedPackagesRaw, ESlateCheckBoxState::Unchecked );
+				PackagesDialogModule.GetResults( UncheckedPackagesRaw, ECheckBoxState::Unchecked );
 				UncheckedPackages.Empty();
 				for (auto Iter = UncheckedPackagesRaw.CreateIterator(); Iter; ++Iter)
 				{
diff --git a/Engine/Source/Editor/UnrealEd/Private/ObjectTools.cpp b/Engine/Source/Editor/UnrealEd/Private/ObjectTools.cpp
--- a/Engine/Source/Editor/UnrealEd/Private/ObjectTools.cpp
+++ b/Engine/Source/Editor/UnrealEd/Private/ObjectTools.cpp
@@ -1303,7 +1303,7 @@ namespace ObjectTools
 			for ( int32 PackageIdx = 0; PackageIdx < PackagesToDelete.Num(); ++PackageIdx )
 			{
 				UPackage* Package = PackagesToDelete[PackageIdx];
-				PackagesDialogModule.AddPackageItem(Package, Package->GetName(), ESlateCheckBoxState::Checked);
+				PackagesDialogModule.AddPackageItem(Package, Package->GetName(), ECheckBoxState::Checked);
 			}
 
 			// Display the delete dialog
diff --git a/Engine/Source/Editor/UnrealEd/Private/PackageRestore.cpp b/Engine/Source/Editor/UnrealEd/Private/PackageRestore.cpp
--- a/Engine/Source/Editor/UnrealEd/Private/PackageRestore.cpp
+++ b/Engine/Source/Editor/UnrealEd/Private/PackageRestore.cpp
@@ -17,7 +17,7 @@ namespace PackageRestore
 			, PackageFilename(InPackageFilename)
 			, AutoSaveFilename(InAutoSaveFilename)
 			, bIsExistingPackage(bInIsExistingPackage)
-			, State(ESlateCheckBoxState::Unchecked)
+			, State(ECheckBoxState::Unchecked)
 		{
 		}
 
@@ -46,13 +46,13 @@ namespace PackageRestore
 		}
 
 		/** @return The state of this item (checked, unchecked) */
-		ESlateCheckBoxState::Type GetState() const
+		ECheckBoxState GetState() const
 		{
 			return State;
 		}
 
 		/** Set the state of this item (checked, unchecked)  */
-		void SetState(const ESlateCheckBoxState::Type InState)
+		void SetState(const ECheckBoxState InState)
 		{
 			State = InState;
 		}
@@ -97,7 +97,7 @@ namespace PackageRestore
 		FString PackageFilename;
 		FString AutoSaveFilename;
 		bool bIsExistingPackage;
-		ESlateCheckBoxState::Type State;
+		ECheckBoxState State;
 	};
 
 	typedef TSharedPtr<FPackageRestoreItem> FPackageRestoreItemPtr;
@@ -320,20 +320,20 @@ namespace PackageRestore
 		 * @return	the desired toggle state for the ToggleSelectedCheckBox.
 		 * Returns Unchecked, unless all of the selected packages are Checked.
 		 */
-		ESlateCheckBoxState::Type GetToggleSelectedState() const
+		ECheckBoxState GetToggleSelectedState() const
 		{
 			// Default to a Checked state
-			ESlateCheckBoxState::Type PendingState = ESlateCheckBoxState::Checked;
+			ECheckBoxState PendingState = ECheckBoxState::Checked;
 
 			for(auto It = PackageRestoreItems->CreateConstIterator(); It; ++It)
 			{
 				const FPackageRestoreItemPtr& ListItem = *It;
 
-				if(ListItem->GetState() == ESlateCheckBoxState::Unchecked)
+				if(ListItem->GetState() == ECheckBoxState::Unchecked)
 				{
 					// If any package in the selection is Unchecked, then represent the entire set of highlighted packages as Unchecked,
 					// so that the first (user) toggle of ToggleSelectedCheckBox consistently Checks all highlighted packages
-					PendingState = ESlateCheckBoxState::Unchecked;
+					PendingState = ECheckBoxState::Unchecked;
 				}
 			}
 
@@ -344,7 +344,7 @@ namespace PackageRestore
 		 * Toggles the highlighted packages.
 		 * If no packages are explicitly highlighted, toggles all packages in the list.
 		 */
-		void OnToggleSelectedCheckBox(ESlateCheckBoxState::Type InNewState)
+		void OnToggleSelectedCheckBox(ECheckBoxState InNewState)
 		{
 			for(auto It = PackageRestoreItems->CreateIterator(); It; ++It)
 			{
@@ -365,7 +365,7 @@ namespace PackageRestore
 			{
 				const FPackageRestoreItemPtr& ListItem = *It;
 
-				if(ListItem->GetState() == ESlateCheckBoxState::Checked)
+				if(ListItem->GetState() == ECheckBoxState::Checked)
 				{
 					return true;
 				}
@@ -421,7 +421,7 @@ namespace PackageRestore
 			{
 				const FPackageRestoreItemPtr& ListItem = *It;
 
-				if(ListItem->GetState() == ESlateCheckBoxState::Checked)
+				if(ListItem->GetState() == ECheckBoxState::Checked)
 				{
 					SelectedPackageItems.Add(ListItem);
 				}
diff --git a/Engine/Source/Editor/UnrealEd/Private/SCurveEditor.cpp b/Engine/Source/Editor/UnrealEd/Private/SCurveEditor.cpp
--- a/Engine/Source/Editor/UnrealEd/Private/SCurveEditor.cpp
+++ b/Engine/Source/Editor/UnrealEd/Private/SCurveEditor.cpp
@@ -302,14 +302,14 @@ FText SCurveEditor::GetIsCurveVisibleToolTip(TSharedPtr<FCurveViewModel> CurveVi
 		FText::Format(LOCTEXT("ShowFormat", "Show {0} curve"), FText::FromName(CurveViewModel->CurveInfo.CurveName));
 }
 
-ESlateCheckBoxState::Type SCurveEditor::IsCurveVisible(TSharedPtr<FCurveViewModel> CurveViewModel) const
+ECheckBoxState SCurveEditor::IsCurveVisible(TSharedPtr<FCurveViewModel> CurveViewModel) const
 {
-	return CurveViewModel->bIsVisible ? ESlateCheckBoxState::Checked : ESlateCheckBoxState::Unchecked;
+	return CurveViewModel->bIsVisible ? ECheckBoxState::Checked : ECheckBoxState::Unchecked;
 }
 
-void SCurveEditor::OnCurveIsVisibleChanged(ESlateCheckBoxState::Type NewCheckboxState, TSharedPtr<FCurveViewModel> CurveViewModel)
+void SCurveEditor::OnCurveIsVisibleChanged(ECheckBoxState NewCheckboxState, TSharedPtr<FCurveViewModel> CurveViewModel)
 {
-	if (NewCheckboxState == ESlateCheckBoxState::Checked)
+	if (NewCheckboxState == ECheckBoxState::Checked)
 	{
 		CurveViewModel->bIsVisible = true;
 	}
@@ -327,14 +327,14 @@ FText SCurveEditor::GetIsCurveLockedToolTip(TSharedPtr<FCurveViewModel> CurveVie
 		FText::Format(LOCTEXT("LockFormat", "Lock {0} curve for editing"), FText::FromName(CurveViewModel->CurveInfo.CurveName));
 }
 
-ESlateCheckBoxState::Type SCurveEditor::IsCurveLocked(TSharedPtr<FCurveViewModel> CurveViewModel) const
+ECheckBoxState SCurveEditor::IsCurveLocked(TSharedPtr<FCurveViewModel> CurveViewModel) const
 {
-	return CurveViewModel->bIsLocked ? ESlateCheckBoxState::Checked : ESlateCheckBoxState::Unchecked;
+	return CurveViewModel->bIsLocked ? ECheckBoxState::Checked : ECheckBoxState::Unchecked;
 }
 
-void SCurveEditor::OnCurveIsLockedChanged(ESlateCheckBoxState::Type NewCheckboxState, TSharedPtr<FCurveViewModel> CurveViewModel)
+void SCurveEditor::OnCurveIsLockedChanged(ECheckBoxState NewCheckboxState, TSharedPtr<FCurveViewModel> CurveViewModel)
 {
-	if (NewCheckboxState == ESlateCheckBoxState::Checked)
+	if (NewCheckboxState == ECheckBoxState::Checked)
 	{
 		CurveViewModel->bIsLocked = true;
 		RemoveCurveKeysFromSelection(CurveViewModel);
diff --git a/Engine/Source/Editor/UnrealEd/Private/SEditorViewportToolBarButton.cpp b/Engine/Source/Editor/UnrealEd/Private/SEditorViewportToolBarButton.cpp
--- a/Engine/Source/Editor/UnrealEd/Private/SEditorViewportToolBarButton.cpp
+++ b/Engine/Source/Editor/UnrealEd/Private/SEditorViewportToolBarButton.cpp
@@ -76,7 +76,7 @@ void SEditorViewportToolBarButton::Construct( const FArguments& Declaration )
 	];
 }
 
-void SEditorViewportToolBarButton::OnCheckStateChanged( ESlateCheckBoxState::Type NewCheckedState )
+void SEditorViewportToolBarButton::OnCheckStateChanged( ECheckBoxState NewCheckedState )
 {
 	// When the check state changes (can only happen during clicking in this case) execute our on clicked delegate
 	if(OnClickedDelegate.IsBound() == true)
@@ -90,8 +90,8 @@ const FSlateBrush* SEditorViewportToolBarButton::OnGetButtonImage() const
 	return IsChecked.Get() == true ? CheckedBrush : NormalBrush;
 }
 
-ESlateCheckBoxState::Type SEditorViewportToolBarButton::OnIsChecked() const
+ECheckBoxState SEditorViewportToolBarButton::OnIsChecked() const
 {
-	return IsChecked.Get() == true ? ESlateCheckBoxState::Checked : ESlateCheckBoxState::Unchecked;
+	return IsChecked.Get() == true ? ECheckBoxState::Checked : ECheckBoxState::Unchecked;
 }
 
diff --git a/Engine/Source/Editor/UnrealEd/Private/SScalabilitySettings.cpp b/Engine/Source/Editor/UnrealEd/Private/SScalabilitySettings.cpp
--- a/Engine/Source/Editor/UnrealEd/Private/SScalabilitySettings.cpp
+++ b/Engine/Source/Editor/UnrealEd/Private/SScalabilitySettings.cpp
@@ -5,7 +5,7 @@
 
 #define LOCTEXT_NAMESPACE "EngineScalabiltySettings"
 
-ESlateCheckBoxState::Type SScalabilitySettings::IsGroupQualityLevelSelected(const TCHAR* InGroupName, int32 InQualityLevel) const
+ECheckBoxState SScalabilitySettings::IsGroupQualityLevelSelected(const TCHAR* InGroupName, int32 InQualityLevel) const
 {
 	int32 QualityLevel = -1;
 
@@ -17,10 +17,10 @@ ESlateCheckBoxState::Type SScalabilitySettings::IsGroupQualityLevelSelected(cons
 	else if (FCString::Strcmp(InGroupName, TEXT("TextureQuality")) == 0) QualityLevel = CachedQualityLevels.TextureQuality;
 	else if (FCString::Strcmp(InGroupName, TEXT("EffectsQuality")) == 0) QualityLevel = CachedQualityLevels.EffectsQuality;
 
-	return (QualityLevel == InQualityLevel) ? ESlateCheckBoxState::Checked : ESlateCheckBoxState::Unchecked;
+	return (QualityLevel == InQualityLevel) ? ECheckBoxState::Checked : ECheckBoxState::Unchecked;
 }
 
-void SScalabilitySettings::OnGroupQualityLevelChanged(ESlateCheckBoxState::Type NewState, const TCHAR* InGroupName, int32 InQualityLevel)
+void SScalabilitySettings::OnGroupQualityLevelChanged(ECheckBoxState NewState, const TCHAR* InGroupName, int32 InQualityLevel)
 {
 	if (FCString::Strcmp(InGroupName, TEXT("ResolutionQuality")) == 0) CachedQualityLevels.ResolutionQuality = InQualityLevel;
 	else if (FCString::Strcmp(InGroupName, TEXT("ViewDistanceQuality")) == 0) CachedQualityLevels.ViewDistanceQuality = InQualityLevel;
@@ -122,15 +122,15 @@ SGridPanel::FSlot& SScalabilitySettings::MakeGridSlot(int32 InCol, int32 InRow,
 		.ColumnSpan(InColSpan);
 }
 
-ESlateCheckBoxState::Type SScalabilitySettings::IsMonitoringPerformance() const
+ECheckBoxState SScalabilitySettings::IsMonitoringPerformance() const
 {
 	const bool bMonitorEditorPerformance = GEditor->GetEditorUserSettings().bMonitorEditorPerformance;
-	return bMonitorEditorPerformance ? ESlateCheckBoxState::Checked : ESlateCheckBoxState::Unchecked;
+	return bMonitorEditorPerformance ? ECheckBoxState::Checked : ECheckBoxState::Unchecked;
 }
 
-void SScalabilitySettings::OnMonitorPerformanceChanged(ESlateCheckBoxState::Type NewState)
+void SScalabilitySettings::OnMonitorPerformanceChanged(ECheckBoxState NewState)
 {
-	const bool bNewEnabledState = ( NewState == ESlateCheckBoxState::Checked );
+	const bool bNewEnabledState = ( NewState == ECheckBoxState::Checked );
 
 	auto& Settings = GEditor->AccessEditorUserSettings();
 	Settings.bMonitorEditorPerformance = bNewEnabledState;
diff --git a/Engine/Source/Editor/UnrealEd/Private/SSkeletonWidget.cpp b/Engine/Source/Editor/UnrealEd/Private/SSkeletonWidget.cpp
--- a/Engine/Source/Editor/UnrealEd/Private/SSkeletonWidget.cpp
+++ b/Engine/Source/Editor/UnrealEd/Private/SSkeletonWidget.cpp
@@ -631,24 +631,24 @@ void SAnimationRemapSkeleton::Construct( const FArguments& InArgs )
 	UpdateAssetPicker();
 }
 
-ESlateCheckBoxState::Type SAnimationRemapSkeleton::IsRemappingReferencedAssets() const
+ECheckBoxState SAnimationRemapSkeleton::IsRemappingReferencedAssets() const
 {
-	return bRemapReferencedAssets ? ESlateCheckBoxState::Checked : ESlateCheckBoxState::Unchecked;
+	return bRemapReferencedAssets ? ECheckBoxState::Checked : ECheckBoxState::Unchecked;
 }
 
-void SAnimationRemapSkeleton::OnRemappingReferencedAssetsChanged(ESlateCheckBoxState::Type InNewRadioState)
+void SAnimationRemapSkeleton::OnRemappingReferencedAssetsChanged(ECheckBoxState InNewRadioState)
 {
-	bRemapReferencedAssets = (InNewRadioState == ESlateCheckBoxState::Checked);
+	bRemapReferencedAssets = (InNewRadioState == ECheckBoxState::Checked);
 }
 
-ESlateCheckBoxState::Type SAnimationRemapSkeleton::IsConvertSpacesChecked() const
+ECheckBoxState SAnimationRemapSkeleton::IsConvertSpacesChecked() const
 {
-	return bConvertSpaces ? ESlateCheckBoxState::Checked : ESlateCheckBoxState::Unchecked;
+	return bConvertSpaces ? ECheckBoxState::Checked : ECheckBoxState::Unchecked;
 }
 
-void SAnimationRemapSkeleton::OnConvertSpacesCheckChanged(ESlateCheckBoxState::Type InNewRadioState)
+void SAnimationRemapSkeleton::OnConvertSpacesCheckChanged(ECheckBoxState InNewRadioState)
 {
-	bConvertSpaces = (InNewRadioState == ESlateCheckBoxState::Checked);
+	bConvertSpaces = (InNewRadioState == ECheckBoxState::Checked);
 }
 
 bool SAnimationRemapSkeleton::IsShowOnlyCompatibleSkeletonsEnabled() const
@@ -657,14 +657,14 @@ bool SAnimationRemapSkeleton::IsShowOnlyCompatibleSkeletonsEnabled() const
 	return (bConvertSpaces == true);
 }
 
-ESlateCheckBoxState::Type SAnimationRemapSkeleton::IsShowOnlyCompatibleSkeletonsChecked() const
+ECheckBoxState SAnimationRemapSkeleton::IsShowOnlyCompatibleSkeletonsChecked() const
 {
-	return bShowOnlyCompatibleSkeletons ? ESlateCheckBoxState::Checked : ESlateCheckBoxState::Unchecked;
+	return bShowOnlyCompatibleSkeletons ? ECheckBoxState::Checked : ECheckBoxState::Unchecked;
 }
 
-void SAnimationRemapSkeleton::OnShowOnlyCompatibleSkeletonsCheckChanged(ESlateCheckBoxState::Type InNewRadioState)
+void SAnimationRemapSkeleton::OnShowOnlyCompatibleSkeletonsCheckChanged(ECheckBoxState InNewRadioState)
 {
-	bShowOnlyCompatibleSkeletons = (InNewRadioState == ESlateCheckBoxState::Checked);
+	bShowOnlyCompatibleSkeletons = (InNewRadioState == ECheckBoxState::Checked);
 
 	UpdateAssetPicker();
 }
diff --git a/Engine/Source/Editor/UnrealEd/Private/STransformViewportToolbar.cpp b/Engine/Source/Editor/UnrealEd/Private/STransformViewportToolbar.cpp
--- a/Engine/Source/Editor/UnrealEd/Private/STransformViewportToolbar.cpp
+++ b/Engine/Source/Editor/UnrealEd/Private/STransformViewportToolbar.cpp
@@ -114,7 +114,7 @@ TSharedRef< SWidget > STransformViewportToolBar::MakeSurfaceSnappingButton( FNam
 			.Style( FEditorStyle::Get(), EMultiBlockLocation::ToName(FEditorStyle::Join(ToolBarStyle, ".ToggleButton"), EMultiBlockLocation::Start) )
 			.Padding( 0 )
 			.ToolTipText( LOCTEXT("SurfaceSnappingCheckboxDescription", "Open editor surface snapping options") )
-			.IsChecked_Static( []{ return GetDefault<ULevelEditorViewportSettings>()->SnapToSurface.bEnabled ? ESlateCheckBoxState::Checked : ESlateCheckBoxState::Unchecked; })
+			.IsChecked_Static( []{ return GetDefault<ULevelEditorViewportSettings>()->SnapToSurface.bEnabled ? ECheckBoxState::Checked : ECheckBoxState::Unchecked; })
 			.Content()
 			[
 				SNew( SComboButton )
@@ -651,32 +651,32 @@ TSharedRef<SWidget> STransformViewportToolBar::FillScaleGridSnapMenu()
 	return ScaleGridMenuBuilder.MakeWidget();
 }
 
-ESlateCheckBoxState::Type STransformViewportToolBar::IsLocationGridSnapChecked() const
+ECheckBoxState STransformViewportToolBar::IsLocationGridSnapChecked() const
 {
-	return GetDefault<ULevelEditorViewportSettings>()->GridEnabled ? ESlateCheckBoxState::Checked : ESlateCheckBoxState::Unchecked;
+	return GetDefault<ULevelEditorViewportSettings>()->GridEnabled ? ECheckBoxState::Checked : ECheckBoxState::Unchecked;
 }
 
-ESlateCheckBoxState::Type STransformViewportToolBar::IsRotationGridSnapChecked() const
+ECheckBoxState STransformViewportToolBar::IsRotationGridSnapChecked() const
 {
-	return GetDefault<ULevelEditorViewportSettings>()->RotGridEnabled ? ESlateCheckBoxState::Checked : ESlateCheckBoxState::Unchecked;
+	return GetDefault<ULevelEditorViewportSettings>()->RotGridEnabled ? ECheckBoxState::Checked : ECheckBoxState::Unchecked;
 }
 
-ESlateCheckBoxState::Type STransformViewportToolBar::IsScaleGridSnapChecked() const
+ECheckBoxState STransformViewportToolBar::IsScaleGridSnapChecked() const
 {
-	return GetDefault<ULevelEditorViewportSettings>()->SnapScaleEnabled ? ESlateCheckBoxState::Checked : ESlateCheckBoxState::Unchecked;
+	return GetDefault<ULevelEditorViewportSettings>()->SnapScaleEnabled ? ECheckBoxState::Checked : ECheckBoxState::Unchecked;
 }
 
-void STransformViewportToolBar::HandleToggleLocationGridSnap( ESlateCheckBoxState::Type InState )
+void STransformViewportToolBar::HandleToggleLocationGridSnap( ECheckBoxState InState )
 {
 	GUnrealEd->Exec( GEditor->GetEditorWorldContext().World(), *FString::Printf( TEXT("MODE GRID=%d"), !GetDefault<ULevelEditorViewportSettings>()->GridEnabled ? 1 : 0 ) );
 }
 
-void STransformViewportToolBar::HandleToggleRotationGridSnap( ESlateCheckBoxState::Type InState )
+void STransformViewportToolBar::HandleToggleRotationGridSnap( ECheckBoxState InState )
 {
 	GUnrealEd->Exec( GEditor->GetEditorWorldContext().World(), *FString::Printf( TEXT("MODE ROTGRID=%d"), !GetDefault<ULevelEditorViewportSettings>()->RotGridEnabled ? 1 : 0 ) );
 }
 
-void STransformViewportToolBar::HandleToggleScaleGridSnap( ESlateCheckBoxState::Type InState )
+void STransformViewportToolBar::HandleToggleScaleGridSnap( ECheckBoxState InState )
 {
 	GUnrealEd->Exec( GEditor->GetEditorWorldContext().World(), *FString::Printf( TEXT("MODE SCALEGRID=%d"), !GetDefault<ULevelEditorViewportSettings>()->SnapScaleEnabled ? 1 : 0 ) );
 }
diff --git a/Engine/Source/Editor/UnrealEd/Private/Settings/ProjectPackagingSettingsCustomization.h b/Engine/Source/Editor/UnrealEd/Private/Settings/ProjectPackagingSettingsCustomization.h
--- a/Engine/Source/Editor/UnrealEd/Private/Settings/ProjectPackagingSettingsCustomization.h
+++ b/Engine/Source/Editor/UnrealEd/Private/Settings/ProjectPackagingSettingsCustomization.h
@@ -273,19 +273,19 @@ protected:
 		return RawCultureStringArray->Num() ? EVisibility::Hidden : EVisibility::Visible;
 	}
 
-	ESlateCheckBoxState::Type HandleShowCulturesCheckBoxIsChecked( EFilterCulturesChoices Choice ) const
+	ECheckBoxState HandleShowCulturesCheckBoxIsChecked( EFilterCulturesChoices Choice ) const
 	{
 		if (FilterCulturesChoice == Choice)
 		{
-			return ESlateCheckBoxState::Checked;
+			return ECheckBoxState::Checked;
 		}
 
-		return ESlateCheckBoxState::Unchecked;
+		return ECheckBoxState::Unchecked;
 	}
 
-	void HandleShowCulturesCheckBoxCheckStateChanged( ESlateCheckBoxState::Type NewState, EFilterCulturesChoices Choice )
+	void HandleShowCulturesCheckBoxCheckStateChanged( ECheckBoxState NewState, EFilterCulturesChoices Choice )
 	{
-		if (NewState == ESlateCheckBoxState::Checked)
+		if (NewState == ECheckBoxState::Checked)
 		{
 			FilterCulturesChoice = Choice;
 		}
diff --git a/Engine/Source/Editor/UnrealEd/Private/Toolkits/AssetEditorManager.cpp b/Engine/Source/Editor/UnrealEd/Private/Toolkits/AssetEditorManager.cpp
--- a/Engine/Source/Editor/UnrealEd/Private/Toolkits/AssetEditorManager.cpp
+++ b/Engine/Source/Editor/UnrealEd/Private/Toolkits/AssetEditorManager.cpp
@@ -519,16 +519,16 @@ void FAssetEditorManager::SpawnRestorePreviouslyOpenAssetsNotification(const boo
 	/** Utility functions for the notification which don't rely on the state from FAssetEditorManager */
 	struct Local
 	{
-		static ESlateCheckBoxState::Type GetDontAskAgainCheckBoxState()
+		static ECheckBoxState GetDontAskAgainCheckBoxState()
 		{
 			bool bSuppressNotification = false;
 			GConfig->GetBool(TEXT("AssetEditorManager"), TEXT("SuppressRestorePreviouslyOpenAssetsNotification"), bSuppressNotification, GEditorUserSettingsIni);
-			return bSuppressNotification ? ESlateCheckBoxState::Checked : ESlateCheckBoxState::Unchecked;
+			return bSuppressNotification ? ECheckBoxState::Checked : ECheckBoxState::Unchecked;
 		}
 
-		static void OnDontAskAgainCheckBoxStateChanged(ESlateCheckBoxState::Type NewState)
+		static void OnDontAskAgainCheckBoxStateChanged(ECheckBoxState NewState)
 		{
-			const bool bSuppressNotification = (NewState == ESlateCheckBoxState::Checked);
+			const bool bSuppressNotification = (NewState == ECheckBoxState::Checked);
 			GConfig->SetBool(TEXT("AssetEditorManager"), TEXT("SuppressRestorePreviouslyOpenAssetsNotification"), bSuppressNotification, GEditorUserSettingsIni);
 		}
 	};
@@ -564,7 +564,7 @@ void FAssetEditorManager::SpawnRestorePreviouslyOpenAssetsNotification(const boo
 	// Only let the user suppress the non-crash version
 	if(bCleanShutdown)
 	{
-		Info.CheckBoxState = TAttribute<ESlateCheckBoxState::Type>::Create(&Local::GetDontAskAgainCheckBoxState);
+		Info.CheckBoxState = TAttribute<ECheckBoxState>::Create(&Local::GetDontAskAgainCheckBoxState);
 		Info.CheckBoxStateChanged = FOnCheckStateChanged::CreateStatic(&Local::OnDontAskAgainCheckBoxStateChanged);
 		Info.CheckBoxText = NSLOCTEXT("ModalDialogs", "DefaultCheckBoxMessage", "Don't show this again");
 	}
diff --git a/Engine/Source/Editor/UnrealEd/Public/AnimationCompressionPanel.h b/Engine/Source/Editor/UnrealEd/Public/AnimationCompressionPanel.h
--- a/Engine/Source/Editor/UnrealEd/Public/AnimationCompressionPanel.h
+++ b/Engine/Source/Editor/UnrealEd/Public/AnimationCompressionPanel.h
@@ -87,11 +87,11 @@ private:
 	 * @param	ButtonId	The ID for the radio button
 	 * @return				The status of the radio button
 	 */
-	ESlateCheckBoxState::Type IsRadioChecked( int32 ButtonId ) const
+	ECheckBoxState IsRadioChecked( int32 ButtonId ) const
 	{
 		return (CurrentCompressionChoice == ButtonId)
-			? ESlateCheckBoxState::Checked
-			: ESlateCheckBoxState::Unchecked;
+			? ECheckBoxState::Checked
+			: ECheckBoxState::Unchecked;
 	}
 
 	/**
@@ -100,9 +100,9 @@ private:
 	 * @param	NewRadioState	The new state of the radio button
 	 * @param	RadioThatChanged	The ID of the radio button that has changed. 
 	 */
-	void OnRadioChanged( ESlateCheckBoxState::Type NewRadioState, int32 RadioThatChanged )
+	void OnRadioChanged( ECheckBoxState NewRadioState, int32 RadioThatChanged )
 	{
-		if (NewRadioState == ESlateCheckBoxState::Checked)
+		if (NewRadioState == ECheckBoxState::Checked)
 		{
 			CurrentCompressionChoice = RadioThatChanged;
 		}
diff --git a/Engine/Source/Editor/UnrealEd/Public/Dialogs/SMeshMergingDialog.h b/Engine/Source/Editor/UnrealEd/Public/Dialogs/SMeshMergingDialog.h
--- a/Engine/Source/Editor/UnrealEd/Public/Dialogs/SMeshMergingDialog.h
+++ b/Engine/Source/Editor/UnrealEd/Public/Dialogs/SMeshMergingDialog.h
@@ -34,8 +34,8 @@ private:
 	FReply OnMergeClicked();
 
 	/**  */
-	ESlateCheckBoxState::Type GetGenerateLightmapUV() const;
-	void SetGenerateLightmapUV(ESlateCheckBoxState::Type NewValue);
+	ECheckBoxState GetGenerateLightmapUV() const;
+	void SetGenerateLightmapUV(ECheckBoxState NewValue);
 
 	/** Target lightmap channel */
 	bool IsLightmapChannelEnabled() const;
@@ -43,15 +43,15 @@ private:
 	void SetTargetLightMapResolution(TSharedPtr<FString> NewSelection, ESelectInfo::Type SelectInfo);
 		
 	/**  */
-	ESlateCheckBoxState::Type GetImportVertexColors() const;
-	void SetImportVertexColors(ESlateCheckBoxState::Type NewValue);
+	ECheckBoxState GetImportVertexColors() const;
+	void SetImportVertexColors(ECheckBoxState NewValue);
 
 	/**  */
-	ESlateCheckBoxState::Type GetPivotPointAtZero() const;
-	void SetPivotPointAtZero(ESlateCheckBoxState::Type NewValue);
+	ECheckBoxState GetPivotPointAtZero() const;
+	void SetPivotPointAtZero(ECheckBoxState NewValue);
 
-	ESlateCheckBoxState::Type GetPlaceInWorld() const;
-	void SetPlaceInWorld(ESlateCheckBoxState::Type NewValue);
+	ECheckBoxState GetPlaceInWorld() const;
+	void SetPlaceInWorld(ECheckBoxState NewValue);
 	
 	/** Called when actors selection is changed */
 	void OnActorSelectionChanged(const TArray<UObject*>& NewSelection);
diff --git a/Engine/Source/Editor/UnrealEd/Public/SCurveEditor.h b/Engine/Source/Editor/UnrealEd/Public/SCurveEditor.h
--- a/Engine/Source/Editor/UnrealEd/Public/SCurveEditor.h
+++ b/Engine/Source/Editor/UnrealEd/Public/SCurveEditor.h
@@ -452,12 +452,12 @@ private:
 	FVector2D SnapLocation(FVector2D InLocation);
 
 	FText GetIsCurveVisibleToolTip(TSharedPtr<FCurveViewModel> CurveViewModel) const;
-	ESlateCheckBoxState::Type IsCurveVisible(TSharedPtr<FCurveViewModel> CurveViewModel) const;
-	void OnCurveIsVisibleChanged(ESlateCheckBoxState::Type NewCheckboxState, TSharedPtr<FCurveViewModel> CurveViewModel);
+	ECheckBoxState IsCurveVisible(TSharedPtr<FCurveViewModel> CurveViewModel) const;
+	void OnCurveIsVisibleChanged(ECheckBoxState NewCheckboxState, TSharedPtr<FCurveViewModel> CurveViewModel);
 
 	FText GetIsCurveLockedToolTip(TSharedPtr<FCurveViewModel> CurveViewModel) const;
-	ESlateCheckBoxState::Type IsCurveLocked(TSharedPtr<FCurveViewModel> CurveViewModel) const;
-	void OnCurveIsLockedChanged(ESlateCheckBoxState::Type NewCheckboxState, TSharedPtr<FCurveViewModel> CurveViewModel);
+	ECheckBoxState IsCurveLocked(TSharedPtr<FCurveViewModel> CurveViewModel) const;
+	void OnCurveIsLockedChanged(ECheckBoxState NewCheckboxState, TSharedPtr<FCurveViewModel> CurveViewModel);
 
 	void RemoveCurveKeysFromSelection(TSharedPtr<FCurveViewModel> CurveViewModel);
 
diff --git a/Engine/Source/Editor/UnrealEd/Public/SEditorViewportToolBarButton.h b/Engine/Source/Editor/UnrealEd/Public/SEditorViewportToolBarButton.h
--- a/Engine/Source/Editor/UnrealEd/Public/SEditorViewportToolBarButton.h
+++ b/Engine/Source/Editor/UnrealEd/Public/SEditorViewportToolBarButton.h
@@ -34,7 +34,7 @@ private:
 	 * 
 	 * @param NewCheckedState	The new state of the check box
 	 */
-	void OnCheckStateChanged( ESlateCheckBoxState::Type NewCheckedState );
+	void OnCheckStateChanged( ECheckBoxState NewCheckedState );
 
 	/**
 	 * Called when we need to get the image to show in the button
@@ -48,7 +48,7 @@ private:
 	 *
 	 * @return The state of the check box button
 	 */
-	ESlateCheckBoxState::Type OnIsChecked() const;
+	ECheckBoxState OnIsChecked() const;
 
 private:
 	/** Attribute used to get the state of a checkbox */
diff --git a/Engine/Source/Editor/UnrealEd/Public/SScalabilitySettings.h b/Engine/Source/Editor/UnrealEd/Public/SScalabilitySettings.h
--- a/Engine/Source/Editor/UnrealEd/Public/SScalabilitySettings.h
+++ b/Engine/Source/Editor/UnrealEd/Public/SScalabilitySettings.h
@@ -24,10 +24,10 @@ public:
 private:
 
 	// Checks cached quality levels to see if the specified group is at the specified quality level
-	ESlateCheckBoxState::Type IsGroupQualityLevelSelected(const TCHAR* InGroupName, int32 InQualityLevel) const;
+	ECheckBoxState IsGroupQualityLevelSelected(const TCHAR* InGroupName, int32 InQualityLevel) const;
 
 	// Callback for when a particular scalability group has its quality level changed
-	void OnGroupQualityLevelChanged(ESlateCheckBoxState::Type NewState, const TCHAR* InGroupName, int32 InQualityLevel);
+	void OnGroupQualityLevelChanged(ECheckBoxState NewState, const TCHAR* InGroupName, int32 InQualityLevel);
 
 	// Callback for when the resolution scale slider changes
 	void OnResolutionScaleChanged(float InValue);
@@ -57,10 +57,10 @@ private:
 	SGridPanel::FSlot& MakeGridSlot(int32 InCol, int32 InRow, int32 InColSpan = 1, int32 InRowSpan = 1);
 
 	/** Called to get the "Show notification" check box state */
-	ESlateCheckBoxState::Type IsMonitoringPerformance() const;
+	ECheckBoxState IsMonitoringPerformance() const;
 
 	/** Called when the state of the "Show notification" check box changes */
-	void OnMonitorPerformanceChanged(ESlateCheckBoxState::Type NewState);
+	void OnMonitorPerformanceChanged(ECheckBoxState NewState);
 
 private:
 	/* The state of scalability settings at the point of opening the menu*/
diff --git a/Engine/Source/Editor/UnrealEd/Public/SSkeletonWidget.h b/Engine/Source/Editor/UnrealEd/Public/SSkeletonWidget.h
--- a/Engine/Source/Editor/UnrealEd/Public/SSkeletonWidget.h
+++ b/Engine/Source/Editor/UnrealEd/Public/SSkeletonWidget.h
@@ -358,17 +358,17 @@ private:
 	FOnRetargetAnimation OnRetargetAnimationDelegate;
 
 	/** Handlers for check box for remapping assets option */
-	ESlateCheckBoxState::Type IsRemappingReferencedAssets() const;
-	void OnRemappingReferencedAssetsChanged(ESlateCheckBoxState::Type InNewRadioState);
+	ECheckBoxState IsRemappingReferencedAssets() const;
+	void OnRemappingReferencedAssetsChanged(ECheckBoxState InNewRadioState);
 
 	/** Handlers for check box for converting spaces*/
-	ESlateCheckBoxState::Type IsConvertSpacesChecked() const;
-	void OnConvertSpacesCheckChanged(ESlateCheckBoxState::Type InNewRadioState);
+	ECheckBoxState IsConvertSpacesChecked() const;
+	void OnConvertSpacesCheckChanged(ECheckBoxState InNewRadioState);
 
 	/** Handlers for check box for converting spaces*/
-	ESlateCheckBoxState::Type IsShowOnlyCompatibleSkeletonsChecked() const;
+	ECheckBoxState IsShowOnlyCompatibleSkeletonsChecked() const;
 	bool IsShowOnlyCompatibleSkeletonsEnabled() const;
-	void OnShowOnlyCompatibleSkeletonsCheckChanged(ESlateCheckBoxState::Type InNewRadioState);
+	void OnShowOnlyCompatibleSkeletonsCheckChanged(ECheckBoxState InNewRadioState);
 
 	/** should filter asset */
 	bool OnShouldFilterAsset(const class FAssetData& AssetData);
