UE-1384: StaticMesh Editor: Copy Collision from Selected is unclear

diff --git a/Engine/Source/Editor/Kismet/Private/STimelineEditor.cpp b/Engine/Source/Editor/Kismet/Private/STimelineEditor.cpp
--- a/Engine/Source/Editor/Kismet/Private/STimelineEditor.cpp
+++ b/Engine/Source/Editor/Kismet/Private/STimelineEditor.cpp
@@ -1363,11 +1363,11 @@ void STimelineEditor::OnTrackNameCommitted( const FText& StringName, ETextCommit
 
 bool STimelineEditor::IsCurveAssetSelected() const
 {
-	// Note: Cannot call GetContentBrowserSelections() during serialization and GC due to its use of FindObject()
+	// Note: Cannot call GetContentBrowserSelectionClasses() during serialization and GC due to its use of FindObject()
 	if(!GIsSavingPackage && !GIsGarbageCollecting)
 	{
 		TArray<UClass*> SelectionList;
-		GEditor->GetContentBrowserSelections(SelectionList);
+		GEditor->GetContentBrowserSelectionClasses(SelectionList);
 
 		for( int i=0; i<SelectionList.Num(); i++ )
 		{
diff --git a/Engine/Source/Editor/StaticMeshEditor/Private/StaticMeshEditor.cpp b/Engine/Source/Editor/StaticMeshEditor/Private/StaticMeshEditor.cpp
--- a/Engine/Source/Editor/StaticMeshEditor/Private/StaticMeshEditor.cpp
+++ b/Engine/Source/Editor/StaticMeshEditor/Private/StaticMeshEditor.cpp
@@ -408,7 +408,8 @@ void FStaticMeshEditor::BindCommands()
 
 	UICommandList->MapAction(
 		Commands.CopyCollisionFromSelectedMesh,
-		FExecuteAction::CreateSP(this, &FStaticMeshEditor::OnCopyCollisionFromSelectedStaticMesh));
+		FExecuteAction::CreateSP(this, &FStaticMeshEditor::OnCopyCollisionFromSelectedStaticMesh),
+		FCanExecuteAction::CreateSP(this, &FStaticMeshEditor::CanCopyCollisionFromSelectedStaticMesh));
 
 	// Mesh menu
 	UICommandList->MapAction(
@@ -418,7 +419,8 @@ void FStaticMeshEditor::BindCommands()
 
 	UICommandList->MapAction(
 		Commands.ChangeMesh,
-		FExecuteAction::CreateSP(this, &FStaticMeshEditor::OnChangeMesh));
+		FExecuteAction::CreateSP(this, &FStaticMeshEditor::OnChangeMesh),
+		FCanExecuteAction::CreateSP(this, &FStaticMeshEditor::CanChangeMesh));
 
 	// Collision Menu
 	UICommandList->MapAction(
@@ -1407,71 +1409,84 @@ void FStaticMeshEditor::OnConvertBoxToConvexCollision()
 
 void FStaticMeshEditor::OnCopyCollisionFromSelectedStaticMesh()
 {
-	// Find currently selected mesh from content browser
-	FContentBrowserModule& ContentBrowserModule = FModuleManager::LoadModuleChecked<FContentBrowserModule>("ContentBrowser");
-	TArray<FAssetData> SelectedAssetData;
-	ContentBrowserModule.Get().GetSelectedAssets(SelectedAssetData);
-	UStaticMesh * SelectedMesh = NULL;
-	if ( SelectedAssetData.Num() > 0 )
-	{
-		// find the first staticmesh we can find
-		for ( auto AssetIt = SelectedAssetData.CreateConstIterator(); AssetIt; ++AssetIt )
-		{
-			const FAssetData& Asset = (*AssetIt);
-			// if staticmesh
-			if ( Asset.GetClass() == UStaticMesh::StaticClass() )
-			{
-				SelectedMesh = Cast<UStaticMesh>(Asset.GetAsset());
-			}
-		}
-	}
+	UStaticMesh* SelectedMesh = GetFirstSelectedStaticMeshInContentBrowser();
+	check(SelectedMesh && SelectedMesh != StaticMesh && SelectedMesh->BodySetup != NULL);
 
-	// If we have a collision model for this staticmesh, ask if we want to replace it.
-	if (SelectedMesh && SelectedMesh != StaticMesh && 
-		SelectedMesh->BodySetup != NULL)
-	{
-		int32 ShouldReplace = FMessageDialog::Open( EAppMsgType::YesNo, FText::Format( LOCTEXT("CopyCollisionFromMeshPrompt", "Are you sure you want to copy collision from {0}?"), FText::FromString( SelectedMesh->GetName() ) ) );
-		if (ShouldReplace == EAppReturnType::Yes)
-		{
-			UBodySetup* BodySetup = StaticMesh->BodySetup;
+	UBodySetup* BodySetup = StaticMesh->BodySetup;
 
-			ClearSelectedPrims();
+	ClearSelectedPrims();
 
-			// Make sure rendering is done - so we are not changing data being used by collision drawing.
-			FlushRenderingCommands();
+	// Make sure rendering is done - so we are not changing data being used by collision drawing.
+	FlushRenderingCommands();
 
-			// Copy body properties from
-			BodySetup->CopyBodyPropertiesFrom(SelectedMesh->BodySetup);
+	GEditor->BeginTransaction(LOCTEXT("FStaticMeshEditor_CopyCollisionFromSelectedStaticMesh", "Copy Collision from Selected Static Mesh"));
+	BodySetup->Modify();
 
-			// Enable collision, if not already
-			if( !Viewport->GetViewportClient().IsSetShowWireframeCollisionChecked() )
-			{
-				Viewport->GetViewportClient().SetShowWireframeCollision();
-			}
+	// Copy body properties from
+	BodySetup->CopyBodyPropertiesFrom(SelectedMesh->BodySetup);
 
-			// Invalidate physics data and create new meshes
-			BodySetup->InvalidatePhysicsData();
-			BodySetup->CreatePhysicsMeshes(); 
+	// Enable collision, if not already
+	if( !Viewport->GetViewportClient().IsSetShowWireframeCollisionChecked() )
+	{
+		Viewport->GetViewportClient().SetShowWireframeCollision();
+	}
 
-			// Mark static mesh as dirty, to help make sure it gets saved.
-			StaticMesh->MarkPackageDirty();
+	// Invalidate physics data and create new meshes
+	BodySetup->InvalidatePhysicsData();
+	BodySetup->CreatePhysicsMeshes(); 
 
-			// Redraw level editor viewports, in case the asset's collision is visible in a viewport and the viewport isn't set to realtime.
-			// Note: This could be more intelligent and only trigger a redraw if the asset is referenced in the world.
-			GUnrealEd->RedrawLevelEditingViewports();
+	GEditor->EndTransaction();
 
-			// Update views/property windows
-			Viewport->RefreshViewport();
+	// Mark static mesh as dirty, to help make sure it gets saved.
+	StaticMesh->MarkPackageDirty();
 
-			StaticMesh->bCustomizedCollision = true;	//mark the static mesh for collision customization
+	// Redraw level editor viewports, in case the asset's collision is visible in a viewport and the viewport isn't set to realtime.
+	// Note: This could be more intelligent and only trigger a redraw if the asset is referenced in the world.
+	GUnrealEd->RedrawLevelEditingViewports();
+
+	// Update views/property windows
+	Viewport->RefreshViewport();
+
+	StaticMesh->bCustomizedCollision = true;	//mark the static mesh for collision customization
+}
+
+bool FStaticMeshEditor::CanCopyCollisionFromSelectedStaticMesh() const
+{
+	bool CanCopy = false;
+
+	TArray<FAssetData> SelectedAssets;
+	GEditor->GetContentBrowserSelections(SelectedAssets);
+	if(SelectedAssets.Num() == 1)
+	{
+		FAssetData& Asset = SelectedAssets[0];
+		if(Asset.GetClass() == UStaticMesh::StaticClass())
+		{
+			UStaticMesh* SelectedMesh = Cast<UStaticMesh>(Asset.GetAsset());
+			if(SelectedMesh && SelectedMesh != StaticMesh && SelectedMesh->BodySetup != NULL)
+			{
+				CanCopy = true;
+			}
 		}
 	}
-	else
+
+	return CanCopy;
+}
+
+UStaticMesh* FStaticMeshEditor::GetFirstSelectedStaticMeshInContentBrowser() const
+{
+	TArray<FAssetData> SelectedAssets;
+	GEditor->GetContentBrowserSelections(SelectedAssets);
+
+	for(auto& Asset : SelectedAssets)
 	{
-		FNotificationInfo Info( LOCTEXT("InvalidMeshSelectionCollisionCopyFromPrompt", "Invalid asset to copy collision from. Select valid StaticMesh in Content Browser.") );
-		Info.ExpireDuration = 2.0f;
-		FSlateNotificationManager::Get().AddNotification( Info );
+		UStaticMesh* SelectedMesh = Cast<UStaticMesh>(Asset.GetAsset());
+		if(SelectedMesh)
+		{
+			return SelectedMesh;
+		}
 	}
+
+	return NULL;
 }
 
 void FStaticMeshEditor::SetEditorMesh(UStaticMesh* InStaticMesh)
@@ -1538,24 +1553,44 @@ void FStaticMeshEditor::SetEditorMesh(UStaticMesh* InStaticMesh)
 
 void FStaticMeshEditor::OnChangeMesh()
 {
-	FEditorDelegates::LoadSelectedAssetsIfNeeded.Broadcast();
-	UStaticMesh* SelectedMesh = GEditor->GetSelectedObjects()->GetTop<UStaticMesh>();
-	if(SelectedMesh && SelectedMesh != StaticMesh)
-	{
-		RemoveEditingObject(StaticMesh);
-		AddEditingObject(SelectedMesh);
+	UStaticMesh* SelectedMesh = GetFirstSelectedStaticMeshInContentBrowser();
+	check(SelectedMesh != NULL && SelectedMesh != StaticMesh);
 
-		SetEditorMesh(SelectedMesh);
+	RemoveEditingObject(StaticMesh);
+	AddEditingObject(SelectedMesh);
 
-		// Clear selections made on previous mesh
-		ClearSelectedPrims();
-		GetSelectedEdges().Empty();
+	SetEditorMesh(SelectedMesh);
 
-		if(SocketManager.IsValid())
+	// Clear selections made on previous mesh
+	ClearSelectedPrims();
+	GetSelectedEdges().Empty();
+
+	if(SocketManager.IsValid())
+	{
+		SocketManager->UpdateStaticMesh();
+	}
+}
+
+bool FStaticMeshEditor::CanChangeMesh() const
+{
+	bool CanChange = false;
+
+	TArray<FAssetData> SelectedAssets;
+	GEditor->GetContentBrowserSelections(SelectedAssets);
+	if(SelectedAssets.Num() == 1)
+	{
+		FAssetData& Asset = SelectedAssets[0];
+		if(Asset.GetClass() == UStaticMesh::StaticClass())
 		{
-			SocketManager->UpdateStaticMesh();
+			UStaticMesh* SelectedMesh = Cast<UStaticMesh>(Asset.GetAsset());
+			if(SelectedMesh && SelectedMesh != StaticMesh)
+			{
+				CanChange = true;
+			}
 		}
 	}
+
+	return CanChange;
 }
 
 void FStaticMeshEditor::DoDecomp(int32 InMaxHullCount, int32 InMaxHullVerts)
diff --git a/Engine/Source/Editor/StaticMeshEditor/Private/StaticMeshEditor.h b/Engine/Source/Editor/StaticMeshEditor/Private/StaticMeshEditor.h
--- a/Engine/Source/Editor/StaticMeshEditor/Private/StaticMeshEditor.h
+++ b/Engine/Source/Editor/StaticMeshEditor/Private/StaticMeshEditor.h
@@ -179,12 +179,21 @@ private:
 	/** Copy the collision data from the selected static mesh in Content Browser*/
 	void OnCopyCollisionFromSelectedStaticMesh(void);
 
+	/** Whether there is a valid static mesh to copy collision from selected in the content browser */
+	bool CanCopyCollisionFromSelectedStaticMesh() const;
+
+	/** Get the first static mesh selected in the content browser */
+	UStaticMesh* GetFirstSelectedStaticMeshInContentBrowser() const;
+
 	/** Clears the collision data for the static mesh */
 	void OnRemoveCollision(void);
 
 	/** Change the mesh the editor is viewing. */
 	void OnChangeMesh();
 
+	/** Whether there is a static mesh selected in the content browser to change to*/
+	bool CanChangeMesh() const;
+
 	/** Rebuilds the LOD combo list and sets it to "auto", a safe LOD level. */
 	void RegenerateLODComboList();
 
diff --git a/Engine/Source/Editor/UnrealEd/Classes/Editor/EditorEngine.h b/Engine/Source/Editor/UnrealEd/Classes/Editor/EditorEngine.h
--- a/Engine/Source/Editor/UnrealEd/Classes/Editor/EditorEngine.h
+++ b/Engine/Source/Editor/UnrealEd/Classes/Editor/EditorEngine.h
@@ -1652,8 +1652,11 @@ public:
 	/** Returns True is a world info Actor is selected */
 	bool IsWorldSettingsSelected();
 
-	/** Function to return list of assets currently selected in content browser (loaded/not loaded) */
-	void GetContentBrowserSelections( TArray<UClass*>& Selection ) const;
+	/** Function to return unique list of the classes of the assets currently selected in content browser (loaded/not loaded) */
+	void GetContentBrowserSelectionClasses( TArray<UClass*>& Selection ) const;
+
+	/** Function to return list of assets currently selected in the content browser */
+	void GetContentBrowserSelections(TArray<FAssetData>& Selections) const;
 
 	/**
 	 * Returns an FSelectionIterator that iterates over the set of selected actors.
diff --git a/Engine/Source/Editor/UnrealEd/Private/Editor.cpp b/Engine/Source/Editor/UnrealEd/Private/Editor.cpp
--- a/Engine/Source/Editor/UnrealEd/Private/Editor.cpp
+++ b/Engine/Source/Editor/UnrealEd/Private/Editor.cpp
@@ -283,7 +283,7 @@ USelection* UEditorEngine::GetSelectedObjects() const
 	return PrivateGetSelectedObjects();
 }
 
-void UEditorEngine::GetContentBrowserSelections( TArray<UClass*>& Selection ) const
+void UEditorEngine::GetContentBrowserSelectionClasses(TArray<UClass*>& Selection) const
 {
 	FContentBrowserModule& ContentBrowserModule = FModuleManager::LoadModuleChecked<FContentBrowserModule>("ContentBrowser");
 	TArray<FAssetData> SelectedAssets;
@@ -300,6 +300,12 @@ void UEditorEngine::GetContentBrowserSelections( TArray<UClass*>& Selection ) co
 	}
 }
 
+void UEditorEngine::GetContentBrowserSelections(TArray<FAssetData>& Selection) const
+{
+	FContentBrowserModule& ContentBrowserModule = FModuleManager::LoadModuleChecked<FContentBrowserModule>("ContentBrowser");
+	ContentBrowserModule.Get().GetSelectedAssets(Selection);
+}
+
 USelection* UEditorEngine::GetSelectedSet( const UClass* Class ) const
 {
 	USelection* SelectedSet = GetSelectedActors();
