Added socket snap to vertex and socket metadata.  Merged pull request, adapted to latest build, and added feature to set orientation of the socket to the normal of the selected vertex.
#github Pull request #200

diff --git a/Engine/Source/Editor/EditorStyle/Private/SlateEditorStyle.cpp b/Engine/Source/Editor/EditorStyle/Private/SlateEditorStyle.cpp
--- a/Engine/Source/Editor/EditorStyle/Private/SlateEditorStyle.cpp
+++ b/Engine/Source/Editor/EditorStyle/Private/SlateEditorStyle.cpp
@@ -5465,6 +5465,8 @@ void FSlateEditorStyle::FStyle::SetupToolkitStyles()
 		Set( "StaticMeshEditor.SetDrawAdditionalData", new IMAGE_BRUSH( "Icons/icon_StaticMeshEd_AdditionalData_40x", Icon40x40 ) );
 		Set( "StaticMeshEditor.SetDrawAdditionalData.Small", new IMAGE_BRUSH( "Icons/icon_StaticMeshEd_AdditionalData_40x", Icon20x20 ) );
 		Set( "StaticMeshEditor.GroupSection", new BOX_BRUSH( "Common/RoundedSelection_16x", FMargin( 4.0f / 16.0f ) ) );
+		Set( "StaticMeshEditor.SetShowVertices", new IMAGE_BRUSH("Icons/icon_StaticMeshEd_Vertices_40x", Icon40x40 ) );
+		Set( "StaticMeshEditor.SetShowVertices.Small", new IMAGE_BRUSH("Icons/icon_StaticMeshEd_Vertices_40x", Icon20x20 ) );
 	}
 
 	// Skeletal Mesh Editor
diff --git a/Engine/Source/Editor/StaticMeshEditor/Private/SStaticMeshEditorViewport.cpp b/Engine/Source/Editor/StaticMeshEditor/Private/SStaticMeshEditorViewport.cpp
--- a/Engine/Source/Editor/StaticMeshEditor/Private/SStaticMeshEditorViewport.cpp
+++ b/Engine/Source/Editor/StaticMeshEditor/Private/SStaticMeshEditorViewport.cpp
@@ -411,6 +411,12 @@ void SStaticMeshEditorViewport::BindCommands()
 		FExecuteAction::CreateSP( EditorViewportClientRef, &FStaticMeshEditorViewportClient::SetDrawAdditionalData ),
 		FCanExecuteAction(),
 		FIsActionChecked::CreateSP( EditorViewportClientRef, &FStaticMeshEditorViewportClient::IsSetDrawAdditionalData ) );
+
+	CommandList->MapAction(
+		Commands.SetShowVertices,
+		FExecuteAction::CreateSP(EditorViewportClientRef, &FStaticMeshEditorViewportClient::SetDrawVertices ),
+		FCanExecuteAction(),
+		FIsActionChecked::CreateSP(EditorViewportClientRef, &FStaticMeshEditorViewportClient::IsSetDrawVerticesChecked ) );
 }
 
 void SStaticMeshEditorViewport::OnFocusViewportToSelection()
diff --git a/Engine/Source/Editor/StaticMeshEditor/Private/StaticMeshEditor.cpp b/Engine/Source/Editor/StaticMeshEditor/Private/StaticMeshEditor.cpp
--- a/Engine/Source/Editor/StaticMeshEditor/Private/StaticMeshEditor.cpp
+++ b/Engine/Source/Editor/StaticMeshEditor/Private/StaticMeshEditor.cpp
@@ -452,6 +452,7 @@ void FStaticMeshEditor::ExtendToolBar()
 				ToolbarBuilder.AddToolBarButton(FStaticMeshEditorCommands::Get().SetShowNormals);
 				ToolbarBuilder.AddToolBarButton(FStaticMeshEditorCommands::Get().SetShowTangents);
 				ToolbarBuilder.AddToolBarButton(FStaticMeshEditorCommands::Get().SetShowBinormals);
+				ToolbarBuilder.AddToolBarButton(FStaticMeshEditorCommands::Get().SetShowVertices);
 				ToolbarBuilder.AddToolBarButton(FStaticMeshEditorCommands::Get().SetDrawUVs);
 				ToolbarBuilder.AddToolBarButton(FStaticMeshEditorCommands::Get().SetDrawAdditionalData);
 			}
diff --git a/Engine/Source/Editor/StaticMeshEditor/Private/StaticMeshEditorActions.cpp b/Engine/Source/Editor/StaticMeshEditor/Private/StaticMeshEditorActions.cpp
--- a/Engine/Source/Editor/StaticMeshEditor/Private/StaticMeshEditorActions.cpp
+++ b/Engine/Source/Editor/StaticMeshEditor/Private/StaticMeshEditorActions.cpp
@@ -20,7 +20,8 @@ void FStaticMeshEditorCommands::RegisterCommands()
 	UI_COMMAND( SetShowNormals, "Normals", "Toggles display of vertex normals in the Preview Pane.", EUserInterfaceActionType::ToggleButton, FInputGesture() );
 	UI_COMMAND( SetShowTangents, "Tangents", "Toggles display of vertex tangents in the Preview Pane.", EUserInterfaceActionType::ToggleButton, FInputGesture() );
 	UI_COMMAND( SetShowBinormals, "Binormals", "Toggles display of vertex binormals (orthogonal vector to normal and tangent) in the Preview Pane.", EUserInterfaceActionType::ToggleButton, FInputGesture() );
-	UI_COMMAND( SetShowPivot, "Show Pivot", "Display the pivot location of the static mesh.", EUserInterfaceActionType::ToggleButton, FInputGesture() );	
+	UI_COMMAND( SetShowVertices, "Vertices", "Toggles display of vertices in the Preview Pane.", EUserInterfaceActionType::ToggleButton, FInputGesture() ); 
+	UI_COMMAND( SetShowPivot, "Show Pivot", "Display the pivot location of the static mesh.", EUserInterfaceActionType::ToggleButton, FInputGesture() );
 
 	UI_COMMAND( CreateDOP10X, "Add 10DOP-X Simplified Collision", "Generates a new axis-aligned box collision mesh with the 4 X-axis aligned edges beveled (10 total sides) encompassing the static mesh.", EUserInterfaceActionType::Button, FInputGesture() );
 	UI_COMMAND( CreateDOP10Y, "Add 10DOP-Y Simplified Collision", "Generates a new axis-aligned box collision mesh with the 4 Y-axis aligned edges beveled (10 total sides) encompassing the static mesh.", EUserInterfaceActionType::Button, FInputGesture() );
diff --git a/Engine/Source/Editor/StaticMeshEditor/Private/StaticMeshEditorViewportClient.cpp b/Engine/Source/Editor/StaticMeshEditor/Private/StaticMeshEditorViewportClient.cpp
--- a/Engine/Source/Editor/StaticMeshEditor/Private/StaticMeshEditorViewportClient.cpp
+++ b/Engine/Source/Editor/StaticMeshEditor/Private/StaticMeshEditorViewportClient.cpp
@@ -81,6 +81,7 @@ FStaticMeshEditorViewportClient::FStaticMeshEditorViewportClient(TWeakPtr<IStati
 	bDrawBinormals = false;
 	bShowPivot = false;
 	bDrawAdditionalData = true;
+	bDrawVertices = false;
 
 	bManipulating = false;
 
@@ -132,6 +133,22 @@ struct HSMESocketProxy : public HHitProxy
 };
 IMPLEMENT_HIT_PROXY(HSMESocketProxy, HHitProxy);
 
+/**
+ * A hit proxy class for vertices.
+ */
+struct HSMEVertexProxy : public HHitProxy
+{
+	DECLARE_HIT_PROXY();
+
+	uint32		Index;
+
+	HSMEVertexProxy(uint32 InIndex)
+		: HHitProxy( HPP_UI )
+		, Index( InIndex )
+	{}
+};
+IMPLEMENT_HIT_PROXY(HSMEVertexProxy, HHitProxy);
+
 bool FStaticMeshEditorViewportClient::InputWidgetDelta( FViewport* Viewport, EAxisList::Type CurrentAxis, FVector& Drag, FRotator& Rot, FVector& Scale )
 {
 	bool bHandled = false;
@@ -510,7 +527,7 @@ void FStaticMeshEditorViewportClient::Draw(const FSceneView* View,FPrimitiveDraw
 	}
 
 
-	if( bDrawNormals || bDrawTangents || bDrawBinormals )
+	if( bDrawNormals || bDrawTangents || bDrawBinormals || bDrawVertices )
 	{
 		FStaticMeshLODResources& LODModel = StaticMesh->RenderData->LODResources[StaticMeshEditorPtr.Pin()->GetCurrentLODIndex()];
 		FIndexArrayView Indices = LODModel.IndexBuffer.GetArrayView();
@@ -527,6 +544,8 @@ void FStaticMeshEditorViewportClient::Draw(const FSceneView* View,FPrimitiveDraw
 			const FVector& Tangent = LODModel.VertexBuffer.VertexTangentX( Indices[i] ); 
 
 			const float Len = 5.0f;
+			const float BoxLen = 2.0f;
+			const FVector Box(BoxLen);
 
 			if( bDrawNormals )
 			{
@@ -542,6 +561,13 @@ void FStaticMeshEditorViewportClient::Draw(const FSceneView* View,FPrimitiveDraw
 			{
 				PDI->DrawLine( WorldPos, WorldPos+LocalToWorldInverseTranspose.TransformVector( Binormal ).GetSafeNormal() * Len, FLinearColor( 0.0f, 0.0f, 1.0f), SDPG_World );
 			}
+
+			if( bDrawVertices )
+			{								
+				PDI->SetHitProxy(new HSMEVertexProxy(i));
+				DrawWireBox( PDI, FBox(VertexPos - Box, VertexPos + Box), FLinearColor(0.0f, 1.0f, 0.0f), SDPG_World );
+				PDI->SetHitProxy(NULL);								
+			}
 		}	
 	}
 
@@ -854,6 +880,27 @@ void FStaticMeshEditorViewportClient::ProcessClick(class FSceneView& InView, cla
 
 			ClearSelectedPrims = false;
 		}
+		else if (IsSetShowSocketsChecked() && HitProxy->IsA(HSMEVertexProxy::StaticGetType()))
+		{
+			UStaticMeshSocket* Socket = StaticMeshEditorPtr.Pin()->GetSelectedSocket();
+
+			if (Socket)
+			{
+				HSMEVertexProxy* VertexProxy = (HSMEVertexProxy*)HitProxy;
+				TSharedPtr<IStaticMeshEditor> StaticMeshEditor = StaticMeshEditorPtr.Pin();
+				if (StaticMeshEditor.IsValid())
+				{
+					FStaticMeshLODResources& LODModel = StaticMesh->RenderData->LODResources[StaticMeshEditor->GetCurrentLODIndex()];
+					FIndexArrayView Indices = LODModel.IndexBuffer.GetArrayView();
+					const uint32 Index = Indices[VertexProxy->Index];
+
+					Socket->RelativeLocation = LODModel.PositionVertexBuffer.VertexPosition(Index);
+					Socket->RelativeRotation = FRotationMatrix::MakeFromYZ(LODModel.VertexBuffer.VertexTangentZ(Index), LODModel.VertexBuffer.VertexTangentX(Index)).Rotator();
+
+					ClearSelectedSockets = false;
+				}
+			}
+		}
 	}
 	else
 	{
@@ -1218,6 +1265,21 @@ bool FStaticMeshEditorViewportClient::IsSetShowBinormalsChecked() const
 	return bDrawBinormals;
 }
 
+void FStaticMeshEditorViewportClient::SetDrawVertices()
+{
+	bDrawVertices = !bDrawVertices;
+	if (FEngineAnalytics::IsAvailable())
+	{
+		FEngineAnalytics::GetProvider().RecordEvent(TEXT("Editor.Usage.StaticMesh.Toolbar"), TEXT("bDrawVertices"), bDrawVertices ? TEXT("True") : TEXT("False"));
+	}
+	Invalidate();
+}
+
+bool FStaticMeshEditorViewportClient::IsSetDrawVerticesChecked() const
+{
+	return bDrawVertices;
+}
+
 void FStaticMeshEditorViewportClient::SetShowWireframeCollision()
 {
 	bShowCollision = !bShowCollision;
diff --git a/Engine/Source/Editor/StaticMeshEditor/Private/StaticMeshEditorViewportClient.h b/Engine/Source/Editor/StaticMeshEditor/Private/StaticMeshEditorViewportClient.h
--- a/Engine/Source/Editor/StaticMeshEditor/Private/StaticMeshEditorViewportClient.h
+++ b/Engine/Source/Editor/StaticMeshEditor/Private/StaticMeshEditorViewportClient.h
@@ -111,6 +111,12 @@ public:
 	/** Callback for checking the additional data drawing flag. */
 	bool IsSetDrawAdditionalData() const;
 
+	/** Callback for toggling the vertices drawing flag. */
+	void SetDrawVertices();
+
+	/** Callback for checking the vertices drawing flag. */
+	bool IsSetDrawVerticesChecked() const;
+
 protected:
 	// FEditorViewportClient interface
 	virtual void PerspectiveCameraMoved() override;
@@ -137,6 +143,7 @@ private:
 	bool bDrawBinormals;
 	bool bShowPivot;
 	bool bDrawAdditionalData;
+	bool bDrawVertices;
 
 	/** true when the user is manipulating a socket widget. */
 	bool bManipulating;
diff --git a/Engine/Source/Editor/StaticMeshEditor/Public/StaticMeshEditorActions.h b/Engine/Source/Editor/StaticMeshEditor/Public/StaticMeshEditorActions.h
--- a/Engine/Source/Editor/StaticMeshEditor/Public/StaticMeshEditorActions.h
+++ b/Engine/Source/Editor/StaticMeshEditor/Public/StaticMeshEditorActions.h
@@ -41,6 +41,7 @@ public:
 	TSharedPtr< FUICommandInfo > SetShowTangents;
 	TSharedPtr< FUICommandInfo > SetShowBinormals;
 	TSharedPtr< FUICommandInfo > SetShowPivot;
+	TSharedPtr< FUICommandInfo > SetShowVertices;
 
 	// Collision Menu Commands
 	TSharedPtr< FUICommandInfo > CreateDOP10X;
diff --git a/Engine/Source/Runtime/Engine/Classes/Engine/StaticMeshSocket.h b/Engine/Source/Runtime/Engine/Classes/Engine/StaticMeshSocket.h
--- a/Engine/Source/Runtime/Engine/Classes/Engine/StaticMeshSocket.h
+++ b/Engine/Source/Runtime/Engine/Classes/Engine/StaticMeshSocket.h
@@ -27,6 +27,9 @@ class UStaticMeshSocket : public UObject
 	UPROPERTY(EditAnywhere, BlueprintReadWrite, Category=StaticMeshSocket)
 	FVector RelativeScale;
 
+	UPROPERTY(EditAnywhere, BlueprintReadWrite, Category=StaticMeshSocket)
+	FString Tag;
+
 #if WITH_EDITORONLY_DATA
 	UPROPERTY(EditAnywhere, Category=StaticMeshSocket)
 	class UStaticMesh* PreviewStaticMesh;
