Missing header updates

diff --git a/Engine/Source/Editor/StaticMeshEditor/Private/SStaticMeshEditorViewport.h b/Engine/Source/Editor/StaticMeshEditor/Private/SStaticMeshEditorViewport.h
--- a/Engine/Source/Editor/StaticMeshEditor/Private/SStaticMeshEditorViewport.h
+++ b/Engine/Source/Editor/StaticMeshEditor/Private/SStaticMeshEditorViewport.h
@@ -121,7 +121,7 @@ private:
 	/** The scene for this viewport. */
 	FPreviewScene PreviewScene;
 
-	/** Level viewport client */
+	/** Editor viewport client */
 	TSharedPtr<class FStaticMeshEditorViewportClient> EditorViewportClient;
 
 	/** Static mesh being edited */
diff --git a/Engine/Source/Editor/UnrealEd/Public/LevelEditorViewport.h b/Engine/Source/Editor/UnrealEd/Public/LevelEditorViewport.h
--- a/Engine/Source/Editor/UnrealEd/Public/LevelEditorViewport.h
+++ b/Engine/Source/Editor/UnrealEd/Public/LevelEditorViewport.h
@@ -119,7 +119,7 @@ public:
 	/**
 	 * Constructor
 	 */
-	FLevelEditorViewportClient();
+	FLevelEditorViewportClient(const TSharedPtr<class SLevelViewport>& InLevelViewport);
 
 	/**
 	 * Destructor
