Add viewport toolbar to static mesh editor

diff --git a/Engine/Source/Editor/StaticMeshEditor/Private/SStaticMeshEditorViewport.cpp b/Engine/Source/Editor/StaticMeshEditor/Private/SStaticMeshEditorViewport.cpp
--- a/Engine/Source/Editor/StaticMeshEditor/Private/SStaticMeshEditorViewport.cpp
+++ b/Engine/Source/Editor/StaticMeshEditor/Private/SStaticMeshEditorViewport.cpp
@@ -19,6 +19,7 @@
 
 #include "ISocketManager.h"
 #include "StaticMeshEditorViewportClient.h"
+#include "Editor/UnrealEd/Public/STransformViewportToolbar.h"
 
 #include "../Private/GeomFitUtils.h"
 #include "ComponentReregisterContext.h"
@@ -319,7 +320,29 @@ TSharedRef<FEditorViewportClient> SStaticMeshEditorViewport::MakeEditorViewportC
 
 TSharedPtr<SWidget> SStaticMeshEditorViewport::MakeViewportToolbar()
 {
-	return NULL; 
+	TSharedRef< STransformViewportToolBar > ToolBar =
+		SNew(STransformViewportToolBar)
+		.Viewport(SharedThis(this))
+		.CommandList(this->GetCommandList())
+		.IsEnabled(FSlateApplication::Get().GetNormalExecutionAttribute());
+
+	return
+		SNew(SBorder)
+		.BorderImage(FEditorStyle::GetBrush("NoBorder"))
+		// Color and opacity is changed based on whether or not the mouse cursor is hovering over the toolbar area
+		.ColorAndOpacity(ToolBar, &SViewportToolBar::OnGetColorAndOpacity)
+		.ForegroundColor(FEditorStyle::GetSlateColor("DefaultForeground"))
+		[
+			SNew(SVerticalBox)
+			.Visibility(EVisibility::SelfHitTestInvisible)
+			+ SVerticalBox::Slot()
+			.AutoHeight()
+			.VAlign(VAlign_Top)
+			.HAlign(HAlign_Right)
+			[
+				ToolBar
+			]
+		];
 }
 
 EVisibility SStaticMeshEditorViewport::OnGetViewportContentVisibility() const
