UE-1092 EDITOR: Collision creation scaling speed is slow

diff --git a/Engine/Source/Editor/StaticMeshEditor/Private/StaticMeshEditor.cpp b/Engine/Source/Editor/StaticMeshEditor/Private/StaticMeshEditor.cpp
--- a/Engine/Source/Editor/StaticMeshEditor/Private/StaticMeshEditor.cpp
+++ b/Engine/Source/Editor/StaticMeshEditor/Private/StaticMeshEditor.cpp
@@ -400,7 +400,8 @@ void FStaticMeshEditor::BindCommands()
 
 	UICommandList->MapAction(
 		Commands.RemoveCollision,
-		FExecuteAction::CreateSP(this, &FStaticMeshEditor::OnRemoveCollision));
+		FExecuteAction::CreateSP(this, &FStaticMeshEditor::OnRemoveCollision),
+		FCanExecuteAction::CreateSP(this, &FStaticMeshEditor::CanRemoveCollision));
 
 	UICommandList->MapAction(
 		Commands.ConvertBoxesToConvex,
@@ -778,6 +779,9 @@ void FStaticMeshEditor::ScaleSelectedPrims(const FVector& InScale)
 		ModifiedScale = InScale * ((GEditor->GetScaleGridSize() / 100.0f) / GEditor->GetGridSize());
 	}
 
+	//Multiply in estimated size of the mesh so scaling of sphere, box and sphyl is similar speed to other scaling
+	float SimplePrimitiveScaleSpeedFactor = StaticMesh->GetBounds().SphereRadius;
+
 	for (int32 PrimIdx = 0; PrimIdx < SelectedPrims.Num(); PrimIdx++)
 	{
 		const FPrimData& PrimData = SelectedPrims[PrimIdx];
@@ -786,13 +790,13 @@ void FStaticMeshEditor::ScaleSelectedPrims(const FVector& InScale)
 		switch (PrimData.PrimType)
 		{
 		case KPT_Sphere:
-			AggGeom->SphereElems[PrimData.PrimIndex].ScaleElem(InScale, MinPrimSize);
+			AggGeom->SphereElems[PrimData.PrimIndex].ScaleElem(SimplePrimitiveScaleSpeedFactor * ModifiedScale, MinPrimSize);
 			break;
 		case KPT_Box:
-			AggGeom->BoxElems[PrimData.PrimIndex].ScaleElem(InScale, MinPrimSize);
+			AggGeom->BoxElems[PrimData.PrimIndex].ScaleElem(SimplePrimitiveScaleSpeedFactor * ModifiedScale, MinPrimSize);
 			break;
 		case KPT_Sphyl:
-			AggGeom->SphylElems[PrimData.PrimIndex].ScaleElem(InScale, MinPrimSize);
+			AggGeom->SphylElems[PrimData.PrimIndex].ScaleElem(SimplePrimitiveScaleSpeedFactor * ModifiedScale, MinPrimSize);
 			break;
 		case KPT_Convex:
 			AggGeom->ConvexElems[PrimData.PrimIndex].ScaleElem(ModifiedScale, MinPrimSize);
@@ -1253,32 +1257,37 @@ void FStaticMeshEditor::OnCollisionSphyl()
 
 void FStaticMeshEditor::OnRemoveCollision(void)
 {
-	// If we have a collision model for this staticmesh, ask if we want to replace it.
 	UBodySetup* BS = StaticMesh->BodySetup;
-	if (BS != NULL && (BS->AggGeom.GetElementCount() > 0))
-	{
-		int32 ShouldReplace = FMessageDialog::Open(EAppMsgType::YesNo, LOCTEXT("RemoveCollisionPrompt", "Are you sure you want to remove all the collision meshes?"));
-		if (ShouldReplace == EAppReturnType::Yes)
-		{
-			ClearSelectedPrims();
+	check(BS != NULL && BS->AggGeom.GetElementCount() > 0);
 
-			// Make sure rendering is done - so we are not changing data being used by collision drawing.
-			FlushRenderingCommands();
+	ClearSelectedPrims();
 
-			StaticMesh->BodySetup->RemoveSimpleCollision();
+	// Make sure rendering is done - so we are not changing data being used by collision drawing.
+	FlushRenderingCommands();
 
-			// refresh collision change back to staticmesh components
-			RefreshCollisionChange(StaticMesh);
+	GEditor->BeginTransaction(LOCTEXT("FStaticMeshEditor_RemoveCollision", "Remove Collision"));
+	StaticMesh->BodySetup->Modify();
 
-			// Mark staticmesh as dirty, to help make sure it gets saved.
-			StaticMesh->MarkPackageDirty();
+	StaticMesh->BodySetup->RemoveSimpleCollision();
 
-			// Update views/property windows
-			Viewport->RefreshViewport();
+	GEditor->EndTransaction();
 
-			StaticMesh->bCustomizedCollision = true;	//mark the static mesh for collision customization
-		}
-	}
+	// refresh collision change back to staticmesh components
+	RefreshCollisionChange(StaticMesh);
+
+	// Mark staticmesh as dirty, to help make sure it gets saved.
+	StaticMesh->MarkPackageDirty();
+
+	// Update views/property windows
+	Viewport->RefreshViewport();
+
+	StaticMesh->bCustomizedCollision = true;	//mark the static mesh for collision customization
+}
+
+bool FStaticMeshEditor::CanRemoveCollision()
+{
+	UBodySetup* BS = StaticMesh->BodySetup;
+	return (BS != NULL && BS->AggGeom.GetElementCount() > 0);
 }
 
 /** Util for adding vertex to an array if it is not already present. */
@@ -1747,69 +1756,65 @@ void FStaticMeshEditor::DeleteSelectedPrims()
 {
 	if (SelectedPrims.Num() > 0)
 	{
-		//int32 ShouldReplace = FMessageDialog::Open(EAppMsgType::YesNo, LOCTEXT("RemoveSelectedCollisionPrompt", "Are you sure you want to remove the selected collision meshes?"));
-		//if (ShouldReplace == EAppReturnType::Yes)
+		// Sort the selected prims by PrimIndex so when we're deleting them we don't mess up other prims indicies
+		struct FCompareFPrimDataPrimIndex
 		{
-			// Sort the selected prims by PrimIndex so when we're deleting them we don't mess up other prims indicies
-			struct FCompareFPrimDataPrimIndex
+			FORCEINLINE bool operator()(const FPrimData& A, const FPrimData& B) const
 			{
-				FORCEINLINE bool operator()(const FPrimData& A, const FPrimData& B) const
-				{
-					return A.PrimIndex < B.PrimIndex;
-				}
-			};
-			SelectedPrims.Sort(FCompareFPrimDataPrimIndex());
+				return A.PrimIndex < B.PrimIndex;
+			}
+		};
+		SelectedPrims.Sort(FCompareFPrimDataPrimIndex());
 
-			check(StaticMesh->BodySetup);
+		check(StaticMesh->BodySetup);
 
-			FKAggregateGeom* AggGeom = &StaticMesh->BodySetup->AggGeom;
+		FKAggregateGeom* AggGeom = &StaticMesh->BodySetup->AggGeom;
 
-			GEditor->BeginTransaction(LOCTEXT("FStaticMeshEditor_DeleteSelectedPrims", "Delete Collision"));
-			StaticMesh->BodySetup->Modify();
+		GEditor->BeginTransaction(LOCTEXT("FStaticMeshEditor_DeleteSelectedPrims", "Delete Collision"));
+		StaticMesh->BodySetup->Modify();
 
-			for (int32 PrimIdx = SelectedPrims.Num() - 1; PrimIdx >= 0; PrimIdx--)
-			{
-				const FPrimData& PrimData = SelectedPrims[PrimIdx];
+		for (int32 PrimIdx = SelectedPrims.Num() - 1; PrimIdx >= 0; PrimIdx--)
+		{
+			const FPrimData& PrimData = SelectedPrims[PrimIdx];
 
-				check(IsPrimValid(PrimData));
-				switch (PrimData.PrimType)
-				{
-				case KPT_Sphere:
-					AggGeom->SphereElems.RemoveAt(PrimData.PrimIndex);
-					break;
-				case KPT_Box:
-					AggGeom->BoxElems.RemoveAt(PrimData.PrimIndex);
-					break;
-				case KPT_Sphyl:
-					AggGeom->SphylElems.RemoveAt(PrimData.PrimIndex);
-					break;
-				case KPT_Convex:
-					AggGeom->ConvexElems.RemoveAt(PrimData.PrimIndex);
-					break;
-				}
+			check(IsPrimValid(PrimData));
+			switch (PrimData.PrimType)
+			{
+			case KPT_Sphere:
+				AggGeom->SphereElems.RemoveAt(PrimData.PrimIndex);
+				break;
+			case KPT_Box:
+				AggGeom->BoxElems.RemoveAt(PrimData.PrimIndex);
+				break;
+			case KPT_Sphyl:
+				AggGeom->SphylElems.RemoveAt(PrimData.PrimIndex);
+				break;
+			case KPT_Convex:
+				AggGeom->ConvexElems.RemoveAt(PrimData.PrimIndex);
+				break;
 			}
+		}
 
-			GEditor->EndTransaction();
+		GEditor->EndTransaction();
 
-			ClearSelectedPrims();
+		ClearSelectedPrims();
 
-			// Make sure rendering is done - so we are not changing data being used by collision drawing.
-			FlushRenderingCommands();
+		// Make sure rendering is done - so we are not changing data being used by collision drawing.
+		FlushRenderingCommands();
 
-			// Make sure to invalidate cooked data
-			StaticMesh->BodySetup->InvalidatePhysicsData();
+		// Make sure to invalidate cooked data
+		StaticMesh->BodySetup->InvalidatePhysicsData();
 
-			// refresh collision change back to staticmesh components
-			RefreshCollisionChange(StaticMesh);
+		// refresh collision change back to staticmesh components
+		RefreshCollisionChange(StaticMesh);
 
-			// Mark staticmesh as dirty, to help make sure it gets saved.
-			StaticMesh->MarkPackageDirty();
+		// Mark staticmesh as dirty, to help make sure it gets saved.
+		StaticMesh->MarkPackageDirty();
 
-			// Update views/property windows
-			Viewport->RefreshViewport();
+		// Update views/property windows
+		Viewport->RefreshViewport();
 
-			StaticMesh->bCustomizedCollision = true;	//mark the static mesh for collision customization
-		}
+		StaticMesh->bCustomizedCollision = true;	//mark the static mesh for collision customization
 	}
 }
 
diff --git a/Engine/Source/Editor/StaticMeshEditor/Private/StaticMeshEditor.h b/Engine/Source/Editor/StaticMeshEditor/Private/StaticMeshEditor.h
--- a/Engine/Source/Editor/StaticMeshEditor/Private/StaticMeshEditor.h
+++ b/Engine/Source/Editor/StaticMeshEditor/Private/StaticMeshEditor.h
@@ -188,6 +188,9 @@ private:
 	/** Clears the collision data for the static mesh */
 	void OnRemoveCollision(void);
 
+	/** Whether there is collision to remove from the static mesh */
+	bool CanRemoveCollision();
+
 	/** Change the mesh the editor is viewing. */
 	void OnChangeMesh();
 
