Fixed issue where LODs were deleted when changing the LOD group in the static mesh editor.

diff --git a/Engine/Source/Editor/StaticMeshEditor/Private/StaticMeshEditorTools.cpp b/Engine/Source/Editor/StaticMeshEditor/Private/StaticMeshEditorTools.cpp
--- a/Engine/Source/Editor/StaticMeshEditor/Private/StaticMeshEditorTools.cpp
+++ b/Engine/Source/Editor/StaticMeshEditor/Private/StaticMeshEditorTools.cpp
@@ -1839,18 +1839,14 @@ void FLevelOfDetailSettingsLayout::OnLODGroupChanged(TSharedPtr<FString> NewValu
 			check(Platform);
 			const FStaticMeshLODGroup& GroupSettings = Platform->GetStaticMeshLODSettings().GetLODGroup(NewGroup);
 
-			// Set the number of LODs to the default.
-			LODCount = GroupSettings.GetDefaultNumLODs();
-			if (StaticMesh->SourceModels.Num() > LODCount)
-			{
-				int32 NumToRemove = StaticMesh->SourceModels.Num() - LODCount;
-				StaticMesh->SourceModels.RemoveAt(LODCount, NumToRemove);
-			}
-			while (StaticMesh->SourceModels.Num() < LODCount)
+			// Set the number of LODs to at least the default. If there are already LODs they will be preserved, with default settings of the new LOD group.
+			int32 DefaultLODCount = GroupSettings.GetDefaultNumLODs();
+
+			while (StaticMesh->SourceModels.Num() < DefaultLODCount)
 			{
 				new(StaticMesh->SourceModels) FStaticMeshSourceModel();
 			}
-			check(StaticMesh->SourceModels.Num() == LODCount);
+			LODCount = DefaultLODCount;
 
 			// Set reduction settings to the defaults.
 			for (int32 LODIndex = 0; LODIndex < LODCount; ++LODIndex)
