Fixed socket preview mesh in Static Mesh Editor so that it transforms as the socket is transformed.
- Regression due to CL 2430088; must now update child components' transforms when the parent socket transform changes.

diff --git a/Engine/Source/Editor/StaticMeshEditor/Private/SStaticMeshEditorViewport.cpp b/Engine/Source/Editor/StaticMeshEditor/Private/SStaticMeshEditorViewport.cpp
--- a/Engine/Source/Editor/StaticMeshEditor/Private/SStaticMeshEditorViewport.cpp
+++ b/Engine/Source/Editor/StaticMeshEditor/Private/SStaticMeshEditorViewport.cpp
@@ -158,14 +158,16 @@ void SStaticMeshEditorViewport::UpdatePreviewSocketMeshes()
 					SocketPreviewMeshComponent = NewObject<UStaticMeshComponent>();
 					PreviewScene.AddComponent(SocketPreviewMeshComponent, FTransform::Identity);
 					SocketPreviewMeshComponents.Add(SocketPreviewMeshComponent);
+					SocketPreviewMeshComponent->SnapTo(PreviewMeshComponent, Socket->SocketName);
 				}
 				else
 				{
 					SocketPreviewMeshComponent = SocketPreviewMeshComponents[i];
+					// Force component to world update to take into account the new socket position.
+					SocketPreviewMeshComponent->UpdateComponentToWorld();
 				}
 
 				SocketPreviewMeshComponent->SetStaticMesh(Socket->PreviewStaticMesh);
-				SocketPreviewMeshComponent->SnapTo(PreviewMeshComponent, Socket->SocketName);
 			}
 		}
 	}
