Integrate - MinLOD setting for static mesh assets and an override on components
* Not yet implemented for foliage

diff --git a/Engine/Source/Editor/StaticMeshEditor/Private/StaticMeshEditorTools.cpp b/Engine/Source/Editor/StaticMeshEditor/Private/StaticMeshEditorTools.cpp
--- a/Engine/Source/Editor/StaticMeshEditor/Private/StaticMeshEditorTools.cpp
+++ b/Engine/Source/Editor/StaticMeshEditor/Private/StaticMeshEditorTools.cpp
@@ -1495,6 +1495,26 @@ void FLevelOfDetailSettingsLayout::AddToDetailsPanel( IDetailLayoutBuilder& Deta
 			.OnSelectionChanged(this, &FLevelOfDetailSettingsLayout::OnImportLOD)
 		];
 
+	LODSettingsCategory.AddCustomRow( LOCTEXT("MinLOD", "Minimum LOD") )
+	.NameContent()
+	[
+		SNew(STextBlock)
+		.Font( IDetailLayoutBuilder::GetDetailFont() )
+		.Text(LOCTEXT("MinLOD", "Minimum LOD"))
+	]
+	.ValueContent()
+	[
+		SNew(SSpinBox<int32>)
+		.Font( IDetailLayoutBuilder::GetDetailFont() )
+		.Value(this, &FLevelOfDetailSettingsLayout::GetMinLOD)
+		.OnValueChanged(this, &FLevelOfDetailSettingsLayout::OnMinLODChanged)
+		.OnValueCommitted(this, &FLevelOfDetailSettingsLayout::OnMinLODCommitted)
+		.MinValue(0)
+		.MaxValue(MAX_STATIC_MESH_LODS)
+		.ToolTipText(this, &FLevelOfDetailSettingsLayout::GetMinLODTooltip)
+		.IsEnabled(FLevelOfDetailSettingsLayout::GetLODCount() > 1)
+	];
+
 	// Add Number of LODs slider.
 	const int32 MinAllowedLOD = 1;
 	LODSettingsCategory.AddCustomRow( LOCTEXT("NumberOfLODs", "Number of LODs") )
@@ -2046,4 +2066,35 @@ FText FLevelOfDetailSettingsLayout::GetLODCountTooltip() const
 	return LOCTEXT("LODCountTooltip_Disabled", "Auto mesh reduction is unavailable! Please provide a mesh reduction interface such as Simplygon to use this feature or manually import LOD levels.");
 }
 
+int32 FLevelOfDetailSettingsLayout::GetMinLOD() const
+{
+	UStaticMesh* StaticMesh = StaticMeshEditor.GetStaticMesh();
+	check(StaticMesh);
+
+	return StaticMesh->MinLOD;
+}
+
+void FLevelOfDetailSettingsLayout::OnMinLODChanged(int32 NewValue)
+{
+	UStaticMesh* StaticMesh = StaticMeshEditor.GetStaticMesh();
+	check(StaticMesh);
+
+	{
+		FStaticMeshComponentRecreateRenderStateContext ReregisterContext(StaticMesh,false);
+		StaticMesh->MinLOD = FMath::Clamp<int32>(NewValue, 0, MAX_STATIC_MESH_LODS - 1);
+		StaticMesh->Modify();
+	}
+	StaticMeshEditor.RefreshViewport();
+}
+
+void FLevelOfDetailSettingsLayout::OnMinLODCommitted(int32 InValue, ETextCommit::Type CommitInfo)
+{
+	OnMinLODChanged(InValue);
+}
+
+FText FLevelOfDetailSettingsLayout::GetMinLODTooltip() const
+{
+	return LOCTEXT("MinLODTooltip", "The minimum LOD to use for rendering.  This can be overridden in components.");
+}
+
 #undef LOCTEXT_NAMESPACE
diff --git a/Engine/Source/Editor/StaticMeshEditor/Private/StaticMeshEditorTools.h b/Engine/Source/Editor/StaticMeshEditor/Private/StaticMeshEditorTools.h
--- a/Engine/Source/Editor/StaticMeshEditor/Private/StaticMeshEditorTools.h
+++ b/Engine/Source/Editor/StaticMeshEditor/Private/StaticMeshEditorTools.h
@@ -287,6 +287,10 @@ private:
 	void OnLODCountCommitted(int32 InValue, ETextCommit::Type CommitInfo);
 	int32 GetLODCount() const;
 
+	void OnMinLODChanged(int32 NewValue);
+	void OnMinLODCommitted(int32 InValue, ETextCommit::Type CommitInfo);
+	int32 GetMinLOD() const;
+
 	float GetLODScreenSize(int32 LODIndex)const;
 	FText GetLODScreenSizeTitle(int32 LODIndex) const;
 	bool CanChangeLODScreenSize() const;
@@ -305,6 +309,7 @@ private:
 	void OnImportLOD(TSharedPtr<FString> NewValue, ESelectInfo::Type SelectInfo);
 	void UpdateLODNames();
 	FText GetLODCountTooltip() const;
+	FText GetMinLODTooltip() const;
 
 private:
 
diff --git a/Engine/Source/Editor/StaticMeshEditor/Private/StaticMeshEditorViewportClient.cpp b/Engine/Source/Editor/StaticMeshEditor/Private/StaticMeshEditorViewportClient.cpp
--- a/Engine/Source/Editor/StaticMeshEditor/Private/StaticMeshEditorViewportClient.cpp
+++ b/Engine/Source/Editor/StaticMeshEditor/Private/StaticMeshEditorViewportClient.cpp
@@ -717,7 +717,7 @@ void FStaticMeshEditorViewportClient::DrawCanvas( FViewport& InViewport, FSceneV
 	int32 CurrentLODLevel = StaticMeshEditor->GetCurrentLODLevel();
 	if (CurrentLODLevel == 0)
 	{
-		CurrentLODLevel = ComputeStaticMeshLOD(StaticMesh->RenderData, StaticMeshComponent->Bounds.Origin, StaticMeshComponent->Bounds.SphereRadius, View);
+		CurrentLODLevel = ComputeStaticMeshLOD(StaticMesh->RenderData, StaticMeshComponent->Bounds.Origin, StaticMeshComponent->Bounds.SphereRadius, View, StaticMesh->MinLOD);
 	}
 	else
 	{
diff --git a/Engine/Source/Runtime/Engine/Classes/Components/StaticMeshComponent.h b/Engine/Source/Runtime/Engine/Classes/Components/StaticMeshComponent.h
--- a/Engine/Source/Runtime/Engine/Classes/Components/StaticMeshComponent.h
+++ b/Engine/Source/Runtime/Engine/Classes/Components/StaticMeshComponent.h
@@ -130,6 +130,14 @@ class ENGINE_API UStaticMeshComponent : public UMeshComponent
 	UPROPERTY()
 	int32 PreviousLODLevel;
 
+	/** 
+	 * Specifies the smallest LOD that will be used for this component.  
+	 * If -1 (default), the MinLOD of the static mesh asset will be used instead.
+	 * This is ignored if ForcedLodModel is enabled.
+	 */
+	UPROPERTY(EditAnywhere, AdvancedDisplay, BlueprintReadOnly, Category=LOD)
+	int32 MinLOD;
+
 	UPROPERTY(EditAnywhere, BlueprintReadOnly, Category=StaticMesh, ReplicatedUsing=OnRep_StaticMesh)
 	class UStaticMesh* StaticMesh;
 
diff --git a/Engine/Source/Runtime/Engine/Classes/Engine/StaticMesh.h b/Engine/Source/Runtime/Engine/Classes/Engine/StaticMesh.h
--- a/Engine/Source/Runtime/Engine/Classes/Engine/StaticMesh.h
+++ b/Engine/Source/Runtime/Engine/Classes/Engine/StaticMesh.h
@@ -309,6 +309,10 @@ class UStaticMesh : public UObject, public IInterface_CollisionDataProvider, pub
 	UPROPERTY()
 	FName LODGroup;
 
+	/** Minimum LOD to use for rendering.  This is the default setting for the mesh and can be overridden by component settings. */
+	UPROPERTY()
+	int32 MinLOD;
+
 	/** If true, the distances at which LODs swap are computed automatically. */
 	UPROPERTY()
 	uint32 bAutoComputeLODScreenSize:1;
diff --git a/Engine/Source/Runtime/Engine/Private/SceneManagement.cpp b/Engine/Source/Runtime/Engine/Private/SceneManagement.cpp
--- a/Engine/Source/Runtime/Engine/Private/SceneManagement.cpp
+++ b/Engine/Source/Runtime/Engine/Private/SceneManagement.cpp
@@ -277,7 +277,7 @@ float ComputeBoundsScreenSize( const FVector4& Origin, const float SphereRadius,
 	return FMath::Clamp(ScreenArea / View.ViewRect.Area(), 0.0f, 1.0f);
 }
 
-int8 ComputeStaticMeshLOD( const FStaticMeshRenderData* RenderData, const FVector4& Origin, const float SphereRadius, const FSceneView& View, float FactorScale )
+int8 ComputeStaticMeshLOD( const FStaticMeshRenderData* RenderData, const FVector4& Origin, const float SphereRadius, const FSceneView& View, int32 MinLOD, float FactorScale )
 {
 	const int32 NumLODs = MAX_STATIC_MESH_LODS;
 
@@ -288,11 +288,11 @@ int8 ComputeStaticMeshLOD( const FStaticMeshRenderData* RenderData, const FVecto
 	{
 		if(RenderData->ScreenSize[LODIndex] > ScreenSize)
 		{
-			return LODIndex;
+			return FMath::Max(LODIndex, MinLOD);
 		}
 	}
 
-	return 0;
+	return MinLOD;
 }
 
 int8 ComputeLODForMeshes( const TIndirectArray<class FStaticMesh>& StaticMeshes, const FSceneView& View, const FVector4& Origin, float SphereRadius, int32 ForcedLODLevel, float ScreenSizeScale )
@@ -310,8 +310,10 @@ int8 ComputeLODForMeshes( const TIndirectArray<class FStaticMesh>& StaticMeshes,
 		}
 		LODToRender = FMath::Clamp<int8>(ForcedLODLevel, 0, MaxLOD);
 	}
-	else
+	else if (View.Family->EngineShowFlags.LOD)
 	{
+		int32 MinLODFound = INT_MAX;
+		bool bFoundLOD = false;
 		int32 NumMeshes = StaticMeshes.Num();
 
 		const float ScreenSize = ComputeBoundsScreenSize(Origin, SphereRadius, View);
@@ -320,16 +322,22 @@ int8 ComputeLODForMeshes( const TIndirectArray<class FStaticMesh>& StaticMeshes,
 		{
 			const FStaticMesh& Mesh = StaticMeshes[MeshIndex];
 
-			if(View.Family->EngineShowFlags.LOD == 1)
-			{
-				float MeshScreenSize = Mesh.ScreenSize * ScreenSizeScale;
+			float MeshScreenSize = Mesh.ScreenSize * ScreenSizeScale;
 
-				if(MeshScreenSize >= ScreenSize)
-				{
-					LODToRender = Mesh.LODIndex;
-					break;
-				}
+			if(MeshScreenSize >= ScreenSize)
+			{
+				LODToRender = Mesh.LODIndex;
+				bFoundLOD = true;
+				break;
 			}
+
+			MinLODFound = FMath::Min<int32>(MinLODFound, Mesh.LODIndex);
+		}
+
+		// If no LOD was found matching the screen size, use the lowest in the array instead of LOD 0, to handle non-zero MinLOD
+		if (!bFoundLOD)
+		{
+			LODToRender = MinLODFound;
 		}
 	}
 	return LODToRender;
diff --git a/Engine/Source/Runtime/Engine/Private/StaticMesh.cpp b/Engine/Source/Runtime/Engine/Private/StaticMesh.cpp
--- a/Engine/Source/Runtime/Engine/Private/StaticMesh.cpp
+++ b/Engine/Source/Runtime/Engine/Private/StaticMesh.cpp
@@ -1223,6 +1223,7 @@ UStaticMesh::UStaticMesh(const FObjectInitializer& ObjectInitializer)
 #endif // #if WITH_EDITORONLY_DATA
 	LightMapResolution = 4;
 	LpvBiasMultiplier = 1.0f;
+	MinLOD = 0;
 }
 
 /**
diff --git a/Engine/Source/Runtime/Engine/Private/StaticMeshComponent.cpp b/Engine/Source/Runtime/Engine/Private/StaticMeshComponent.cpp
--- a/Engine/Source/Runtime/Engine/Private/StaticMeshComponent.cpp
+++ b/Engine/Source/Runtime/Engine/Private/StaticMeshComponent.cpp
@@ -144,6 +144,7 @@ UStaticMeshComponent::UStaticMeshComponent(const FObjectInitializer& ObjectIniti
 
 	WireframeColorOverride = FColor(255, 255, 255, 255);
 
+	MinLOD = -1;
 	bOverrideLightMapRes = false;
 	OverriddenLightMapRes = 64;
 	SubDivisionStepSize = 32;
diff --git a/Engine/Source/Runtime/Engine/Private/StaticMeshRender.cpp b/Engine/Source/Runtime/Engine/Private/StaticMeshRender.cpp
--- a/Engine/Source/Runtime/Engine/Private/StaticMeshRender.cpp
+++ b/Engine/Source/Runtime/Engine/Private/StaticMeshRender.cpp
@@ -63,6 +63,8 @@ FStaticMeshSceneProxy::FStaticMeshSceneProxy(UStaticMeshComponent* InComponent):
 {
 	check(RenderData);
 
+	MinLOD = FMath::Clamp(InComponent->MinLOD < 0 ? InComponent->StaticMesh->MinLOD : InComponent->MinLOD, 0, RenderData->LODResources.Num() - 1);
+
 	WireframeColor = InComponent->GetWireframeColor();
 	LevelColor = FLinearColor(1,1,1);
 	PropertyColor = FLinearColor(1,1,1);
@@ -464,7 +466,7 @@ void FStaticMeshSceneProxy::DrawStaticElements(FStaticPrimitiveDrawInterface* PD
 		} 
 		else //no LOD is being forced, submit them all with appropriate cull distances
 		{
-			for(int32 LODIndex = 0; LODIndex < NumLODs; LODIndex++)
+			for(int32 LODIndex = MinLOD; LODIndex < NumLODs; LODIndex++)
 			{
 				const FStaticMeshLODResources& LODModel = RenderData->LODResources[LODIndex];
 				float ScreenSize = GetScreenSize(LODIndex);
@@ -1175,7 +1177,7 @@ int32 FStaticMeshSceneProxy::GetLOD(const FSceneView* View) const
 #endif
 
 	const FBoxSphereBounds& Bounds = GetBounds();
-	return ComputeStaticMeshLOD(RenderData, Bounds.Origin, Bounds.SphereRadius, *View);
+	return ComputeStaticMeshLOD(RenderData, Bounds.Origin, Bounds.SphereRadius, *View, MinLOD);
 }
 
 FPrimitiveSceneProxy* UStaticMeshComponent::CreateSceneProxy()
diff --git a/Engine/Source/Runtime/Engine/Public/SceneManagement.h b/Engine/Source/Runtime/Engine/Public/SceneManagement.h
--- a/Engine/Source/Runtime/Engine/Public/SceneManagement.h
+++ b/Engine/Source/Runtime/Engine/Public/SceneManagement.h
@@ -1968,7 +1968,7 @@ float ENGINE_API ComputeBoundsScreenSize(const FVector4& Origin, const float Sph
  * @param SphereRadius - Radius of the sphere to use to calculate screen coverage
  * @param View - The view to calculate the LOD level for
  */
-int8 ENGINE_API ComputeStaticMeshLOD(const FStaticMeshRenderData* RenderData, const FVector4& Origin, const float SphereRadius, const FSceneView& View, float FactorScale = 1.0f);
+int8 ENGINE_API ComputeStaticMeshLOD(const FStaticMeshRenderData* RenderData, const FVector4& Origin, const float SphereRadius, const FSceneView& View, int32 MinLOD, float FactorScale = 1.0f);
 
 /**
  * Computes the LOD to render for the list of static meshes in the given view.
diff --git a/Engine/Source/Runtime/Engine/Public/StaticMeshResources.h b/Engine/Source/Runtime/Engine/Public/StaticMeshResources.h
--- a/Engine/Source/Runtime/Engine/Public/StaticMeshResources.h
+++ b/Engine/Source/Runtime/Engine/Public/StaticMeshResources.h
@@ -874,6 +874,8 @@ protected:
 	 */
 	int32 ForcedLodModel;
 
+	int32 MinLOD;
+
 	FVector TotalScale3D;
 
 	uint32 bCastShadow : 1;
