Static Mesh Editor: Switched to use the full standard in-viewport toolbar
- Switched transform menu over to use SCommonEditorViewportToolbarBase, which also includes the 'Camera', 'View Mode', Show', etc... menus
- Transform toolbar is now only shown when something that can be modified via the gizmo is selected (collision, socket, etc...)

diff --git a/Engine/Source/Editor/StaticMeshEditor/Private/SStaticMeshEditorViewport.cpp b/Engine/Source/Editor/StaticMeshEditor/Private/SStaticMeshEditorViewport.cpp
--- a/Engine/Source/Editor/StaticMeshEditor/Private/SStaticMeshEditorViewport.cpp
+++ b/Engine/Source/Editor/StaticMeshEditor/Private/SStaticMeshEditorViewport.cpp
@@ -35,6 +35,61 @@
 
 #define HITPROXY_SOCKET	1
 
+///////////////////////////////////////////////////////////
+// SStaticMeshEditorViewportToolbar
+
+// In-viewport toolbar widget used in the static mesh editor
+class SStaticMeshEditorViewportToolbar : public SCommonEditorViewportToolbarBase
+{
+public:
+	SLATE_BEGIN_ARGS(SStaticMeshEditorViewportToolbar) {}
+	SLATE_END_ARGS()
+
+	void Construct(const FArguments& InArgs, TSharedPtr<class ICommonEditorViewportToolbarInfoProvider> InInfoProvider)
+	{
+		SCommonEditorViewportToolbarBase::Construct(SCommonEditorViewportToolbarBase::FArguments(), InInfoProvider);
+	}
+
+	// SCommonEditorViewportToolbarBase interface
+	virtual TSharedRef<SWidget> GenerateShowMenu() const override
+	{
+		GetInfoProvider().OnFloatingButtonClicked();
+
+		TSharedRef<SEditorViewport> ViewportRef = GetInfoProvider().GetViewportWidget();
+
+		const bool bInShouldCloseWindowAfterMenuSelection = true;
+		FMenuBuilder ShowMenuBuilder(bInShouldCloseWindowAfterMenuSelection, ViewportRef->GetCommandList());
+		{
+			auto Commands = FStaticMeshEditorCommands::Get();
+
+			ShowMenuBuilder.AddMenuEntry(Commands.SetShowSockets);
+			ShowMenuBuilder.AddMenuEntry(Commands.SetShowPivot);
+			ShowMenuBuilder.AddMenuEntry(Commands.SetShowVertices);
+
+			ShowMenuBuilder.AddMenuSeparator();
+
+			ShowMenuBuilder.AddMenuEntry(Commands.SetShowGrid);
+			ShowMenuBuilder.AddMenuEntry(Commands.SetShowBounds);
+			ShowMenuBuilder.AddMenuEntry(Commands.SetShowCollision);
+
+			ShowMenuBuilder.AddMenuSeparator();
+
+			ShowMenuBuilder.AddMenuEntry(Commands.SetShowNormals);
+			ShowMenuBuilder.AddMenuEntry(Commands.SetShowTangents);
+			ShowMenuBuilder.AddMenuEntry(Commands.SetShowBinormals);
+
+			//ShowMenuBuilder.AddMenuSeparator();
+			//ShowMenuBuilder.AddMenuEntry(Commands.SetShowMeshEdges);
+		}
+
+		return ShowMenuBuilder.MakeWidget();
+	}
+	// End of SCommonEditorViewportToolbarBase
+};
+
+///////////////////////////////////////////////////////////
+// SStaticMeshEditorViewport
+
 void SStaticMeshEditorViewport::Construct(const FArguments& InArgs)
 {
 	StaticMeshEditorPtr = InArgs._StaticMeshEditor;
@@ -52,7 +107,7 @@ void SStaticMeshEditorViewport::Construct(const FArguments& InArgs)
 	ViewportOverlay->AddSlot()
 		.VAlign(VAlign_Top)
 		.HAlign(HAlign_Left)
-		.Padding(10)
+		.Padding(FMargin(10.0f, 40.0f, 10.0f, 10.0f))
 		[
 			SAssignNew(OverlayTextVerticalBox, SVerticalBox)
 		];
@@ -85,6 +140,21 @@ void SStaticMeshEditorViewport::PopulateOverlayText(const TArray<FOverlayTextIte
 	}
 }
 
+TSharedRef<class SEditorViewport> SStaticMeshEditorViewport::GetViewportWidget()
+{
+	return SharedThis(this);
+}
+
+TSharedPtr<FExtender> SStaticMeshEditorViewport::GetExtenders() const
+{
+	TSharedPtr<FExtender> Result(MakeShareable(new FExtender));
+	return Result;
+}
+
+void SStaticMeshEditorViewport::OnFloatingButtonClicked()
+{
+}
+
 void SStaticMeshEditorViewport::AddReferencedObjects( FReferenceCollector& Collector )
 {
 	Collector.AddReferencedObject( PreviewMeshComponent );
@@ -321,29 +391,7 @@ TSharedRef<FEditorViewportClient> SStaticMeshEditorViewport::MakeEditorViewportC
 
 TSharedPtr<SWidget> SStaticMeshEditorViewport::MakeViewportToolbar()
 {
-	TSharedRef< STransformViewportToolBar > ToolBar =
-		SNew(STransformViewportToolBar)
-		.Viewport(SharedThis(this))
-		.CommandList(this->GetCommandList())
-		.IsEnabled(FSlateApplication::Get().GetNormalExecutionAttribute());
-
-	return
-		SNew(SBorder)
-		.BorderImage(FEditorStyle::GetBrush("NoBorder"))
-		// Color and opacity is changed based on whether or not the mouse cursor is hovering over the toolbar area
-		.ColorAndOpacity(ToolBar, &SViewportToolBar::OnGetColorAndOpacity)
-		.ForegroundColor(FEditorStyle::GetSlateColor("DefaultForeground"))
-		[
-			SNew(SVerticalBox)
-			.Visibility(EVisibility::SelfHitTestInvisible)
-			+ SVerticalBox::Slot()
-			.AutoHeight()
-			.VAlign(VAlign_Top)
-			.HAlign(HAlign_Right)
-			[
-				ToolBar
-			]
-		];
+	return SNew(SStaticMeshEditorViewportToolbar, SharedThis(this));
 }
 
 EVisibility SStaticMeshEditorViewport::OnGetViewportContentVisibility() const
diff --git a/Engine/Source/Editor/StaticMeshEditor/Private/SStaticMeshEditorViewport.h b/Engine/Source/Editor/StaticMeshEditor/Private/SStaticMeshEditorViewport.h
--- a/Engine/Source/Editor/StaticMeshEditor/Private/SStaticMeshEditorViewport.h
+++ b/Engine/Source/Editor/StaticMeshEditor/Private/SStaticMeshEditorViewport.h
@@ -4,11 +4,12 @@
 
 #include "PreviewScene.h"
 #include "SEditorViewport.h"
+#include "Editor/UnrealEd/Public/SCommonEditorViewportToolbarBase.h"
 
 /**
  * StaticMesh Editor Preview viewport widget
  */
-class SStaticMeshEditorViewport : public SEditorViewport, public FGCObject
+class SStaticMeshEditorViewport : public SEditorViewport, public FGCObject, public ICommonEditorViewportToolbarInfoProvider
 {
 public:
 	SLATE_BEGIN_ARGS( SStaticMeshEditorViewport ){}
@@ -81,6 +82,12 @@ public:
 	/** Specifies an array of text items which will be added to the viewport overlay */
 	void PopulateOverlayText( const TArray<FOverlayTextItem>& TextItems );
 
+	// ICommonEditorViewportToolbarInfoProvider interface
+	virtual TSharedRef<class SEditorViewport> GetViewportWidget() override;
+	virtual TSharedPtr<FExtender> GetExtenders() const override;
+	virtual void OnFloatingButtonClicked() override;
+	// End of ICommonEditorViewportToolbarInfoProvider interface
+
 protected:
 	/** SEditorViewport interface */
 	virtual TSharedRef<FEditorViewportClient> MakeEditorViewportClient() override;
