Fixed out of range array access for LOD Group names if LOD group not found.

diff --git a/Engine/Source/Editor/StaticMeshEditor/Private/StaticMeshEditorTools.cpp b/Engine/Source/Editor/StaticMeshEditor/Private/StaticMeshEditorTools.cpp
--- a/Engine/Source/Editor/StaticMeshEditor/Private/StaticMeshEditorTools.cpp
+++ b/Engine/Source/Editor/StaticMeshEditor/Private/StaticMeshEditorTools.cpp
@@ -1464,7 +1464,8 @@ void FLevelOfDetailSettingsLayout::AddToDetailsPanel( IDetailLayoutBuilder& Deta
 	IDetailCategoryBuilder& LODSettingsCategory =
 		DetailBuilder.EditCategory( "LodSettings", LOCTEXT("LodSettingsCategory", "LOD Settings") );
 
-	
+	int32 LODGroupIndex = LODGroupNames.Find(StaticMesh->LODGroup);
+	check(LODGroupIndex == INDEX_NONE || LODGroupIndex < LODGroupOptions.Num());
 
 	LODSettingsCategory.AddCustomRow( LOCTEXT("LODGroup", "LOD Group") )
 	.NameContent()
@@ -1478,7 +1479,7 @@ void FLevelOfDetailSettingsLayout::AddToDetailsPanel( IDetailLayoutBuilder& Deta
 		SNew(STextComboBox)
 		.ContentPadding(0)
 		.OptionsSource(&LODGroupOptions)
-		.InitiallySelectedItem(LODGroupOptions[LODGroupNames.Find(StaticMesh->LODGroup)])
+		.InitiallySelectedItem(LODGroupOptions[(LODGroupIndex == INDEX_NONE) ? 0 : LODGroupIndex])
 		.OnSelectionChanged(this, &FLevelOfDetailSettingsLayout::OnLODGroupChanged)
 	];
 	
