Add back in option to toggle  MikkTSpace tangent space generation for static meshes

diff --git a/Engine/Source/Editor/StaticMeshEditor/Private/StaticMeshEditorTools.cpp b/Engine/Source/Editor/StaticMeshEditor/Private/StaticMeshEditorTools.cpp
--- a/Engine/Source/Editor/StaticMeshEditor/Private/StaticMeshEditorTools.cpp
+++ b/Engine/Source/Editor/StaticMeshEditor/Private/StaticMeshEditorTools.cpp
@@ -315,8 +315,6 @@ void FMeshBuildSettingsLayout::GenerateHeaderRowContent( FDetailWidgetRow& NodeR
 	];
 }
 
-TAutoConsoleVariable<int32> GEnableMikkTSpaceCVar(TEXT("r.MikkTSPaceOptional"),0,TEXT("Set to be non-zero to display the option of using MikkTSpace to generate tangents for static meshes. MikkTSpace is the default."),ECVF_Default);
-
 FString FMeshBuildSettingsLayout::GetCurrentDistanceFieldReplacementMeshPath() const
 {
 	return BuildSettings.DistanceFieldReplacementMesh ? BuildSettings.DistanceFieldReplacementMesh->GetPathName() : FString("");
@@ -362,14 +360,13 @@ void FMeshBuildSettingsLayout::GenerateChildContent( IDetailChildrenBuilder& Chi
 		];
 	}
 
-	if (GEnableMikkTSpaceCVar.GetValueOnAnyThread())
 	{
-		ChildrenBuilder.AddChildContent( LOCTEXT("UseMikkTSpace", "Use MikkTSpace") )
+		ChildrenBuilder.AddChildContent( LOCTEXT("UseMikkTSpace", "Use MikkTSpace Tangent Space") )
 		.NameContent()
 		[
 			SNew(STextBlock)
 			.Font( IDetailLayoutBuilder::GetDetailFont() )
-			.Text(LOCTEXT("UseMikkTSpace", "Use MikkTSpace"))
+			.Text(LOCTEXT("UseMikkTSpace", "Use MikkTSpace Tangent Space"))
 		]
 		.ValueContent()
 		[
