Fixed static mesh editor not using the correct window for its screenshots
This addresses UE-12575

diff --git a/Engine/Source/Editor/StaticMeshEditor/Private/StaticMeshAutomationTests.cpp b/Engine/Source/Editor/StaticMeshEditor/Private/StaticMeshAutomationTests.cpp
--- a/Engine/Source/Editor/StaticMeshEditor/Private/StaticMeshAutomationTests.cpp
+++ b/Engine/Source/Editor/StaticMeshEditor/Private/StaticMeshAutomationTests.cpp
@@ -180,17 +180,22 @@ bool FStaticMeshEditorTest::RunTest(const FString& Parameters)
 		LOAD_None,
 		NULL);
 	FAssetEditorManager::Get().OpenEditorForAsset(EditorMesh);
-	WindowParameters.CurrentWindow = FSlateApplication::Get().GetActiveTopLevelWindow();
 
-	//Check that we have the right window (Tutorial may have opened on top of the editor)
-	if( WindowParameters.CurrentWindow->GetTitle().ToString() != LoadedObjectType )
 	{
-		//Close the tutorial
-		WindowParameters.CurrentWindow->RequestDestroyWindow();
+		TArray<TSharedRef<SWindow>> AllWindows;
+		FSlateApplication::Get().GetAllVisibleWindowsOrdered(AllWindows);
 
-		//Bring the asset editor to the front
-		FAssetEditorManager::Get().FindEditorForAsset(EditorMesh,false);
-		WindowParameters.CurrentWindow = FSlateApplication::Get().GetActiveTopLevelWindow();
+		const FText ExpectedTitle = FText::FromString(LoadedObjectType);
+		if (auto* Window = AllWindows.FindByPredicate([&](const TSharedRef<SWindow>& In) { return In->GetTitle().EqualTo(ExpectedTitle); }))
+		{
+			WindowParameters.CurrentWindow = *Window;
+		}
+	}
+
+	if (!WindowParameters.CurrentWindow.IsValid())
+	{
+		AddError(TEXT("Could not find static mesh editor window"));
+		return false;
 	}
 
 	//Grab the last opened Viewport (aka the asset manager)
