Fix crash when reimporting a mesh while editing collision in the static mesh editor

diff --git a/Engine/Source/Editor/StaticMeshEditor/Private/StaticMeshEditor.cpp b/Engine/Source/Editor/StaticMeshEditor/Private/StaticMeshEditor.cpp
--- a/Engine/Source/Editor/StaticMeshEditor/Private/StaticMeshEditor.cpp
+++ b/Engine/Source/Editor/StaticMeshEditor/Private/StaticMeshEditor.cpp
@@ -1526,6 +1526,8 @@ UStaticMesh* FStaticMeshEditor::GetFirstSelectedStaticMeshInContentBrowser() con
 
 void FStaticMeshEditor::SetEditorMesh(UStaticMesh* InStaticMesh)
 {
+	ClearSelectedPrims();
+
 	StaticMesh = InStaticMesh;
 
 	//Init stat arrays.
