Auto-convex generation now ignores sections with collision disabled (suggestion from [REDACTED_HANDLE])

diff --git a/Engine/Source/Editor/StaticMeshEditor/Private/StaticMeshEditor.cpp b/Engine/Source/Editor/StaticMeshEditor/Private/StaticMeshEditor.cpp
--- a/Engine/Source/Editor/StaticMeshEditor/Private/StaticMeshEditor.cpp
+++ b/Engine/Source/Editor/StaticMeshEditor/Private/StaticMeshEditor.cpp
@@ -1664,9 +1664,22 @@ void FStaticMeshEditor::DoDecomp(float InAccuracy, int32 InMaxHullVerts)
 			Verts.Add(Vert);
 		}
 
-		// Make index buffer
-		TArray<uint32> Indices;
-		LODModel.IndexBuffer.GetCopy(Indices);
+		// Grab all indices
+		TArray<uint32> AllIndices;
+		LODModel.IndexBuffer.GetCopy(AllIndices);
+
+		// Only copy indices that have collision enabled
+		TArray<uint32> CollidingIndices;
+		for(const FStaticMeshSection& Section : LODModel.Sections)
+		{
+			if(Section.bEnableCollision)
+			{
+				for (uint32 IndexIdx = Section.FirstIndex; IndexIdx < Section.FirstIndex + (Section.NumTriangles * 3); IndexIdx++)
+				{
+					CollidingIndices.Add(AllIndices[IndexIdx]);
+				}
+			}
+		}
 
 		ClearSelectedPrims();
 
@@ -1687,7 +1700,7 @@ void FStaticMeshEditor::DoDecomp(float InAccuracy, int32 InMaxHullVerts)
 		}
 
 		// Run actual util to do the work
-		DecomposeMeshToHulls(bs, Verts, Indices, InAccuracy, InMaxHullVerts);		
+		DecomposeMeshToHulls(bs, Verts, CollidingIndices, InAccuracy, InMaxHullVerts);		
 
 		// Enable collision, if not already
 		if( !Viewport->GetViewportClient().IsSetShowWireframeCollisionChecked() )
