Changed static mesh adjacency buffer generation from fixed 50k triangle limit to import option. UE-17556

diff --git a/Engine/Source/Developer/MeshUtilities/Private/MeshUtilities.cpp b/Engine/Source/Developer/MeshUtilities/Private/MeshUtilities.cpp
--- a/Engine/Source/Developer/MeshUtilities/Private/MeshUtilities.cpp
+++ b/Engine/Source/Developer/MeshUtilities/Private/MeshUtilities.cpp
@@ -851,7 +851,7 @@ static void BuildStaticAdjacencyIndexBuffer(
 	TArray<uint32>& OutPnAenIndices
 	)
 {
-	if ( Indices.Num() && Indices.Num() < 50000 * 3 )
+	if ( Indices.Num() )
 	{
 		FStaticMeshNvRenderBuffer StaticMeshRenderBuffer( PositionVertexBuffer, VertexBuffer, Indices );
 		nv::IndexBuffer* PnAENIndexBuffer = nv::tess::buildTessellationBuffer( &StaticMeshRenderBuffer, nv::DBM_PnAenDominantCorner, true );
@@ -3049,6 +3049,7 @@ public:
 			}
 
 			// Build the adjacency index buffer used for tessellation.
+			if (InOutModels[0].BuildSettings.bBuildAdjacencyBuffer)
 			{
 				TArray<uint32> AdjacencyIndices;
 
diff --git a/Engine/Source/Editor/StaticMeshEditor/Private/StaticMeshEditorTools.cpp b/Engine/Source/Editor/StaticMeshEditor/Private/StaticMeshEditorTools.cpp
--- a/Engine/Source/Editor/StaticMeshEditor/Private/StaticMeshEditorTools.cpp
+++ b/Engine/Source/Editor/StaticMeshEditor/Private/StaticMeshEditorTools.cpp
@@ -392,6 +392,22 @@ void FMeshBuildSettingsLayout::GenerateChildContent( IDetailChildrenBuilder& Chi
 		];
 	}
 
+	{
+		ChildrenBuilder.AddChildContent( LOCTEXT("BuildAdjacencyBuffer", "Build Adjacency Buffer") )
+		.NameContent()
+		[
+			SNew(STextBlock)
+			.Font( IDetailLayoutBuilder::GetDetailFont() )
+			.Text(LOCTEXT("BuildAdjacencyBuffer", "Build Adjacency Buffer"))
+		]
+		.ValueContent()
+		[
+			SNew(SCheckBox)
+			.IsChecked(this, &FMeshBuildSettingsLayout::ShouldBuildAdjacencyBuffer)
+			.OnCheckStateChanged(this, &FMeshBuildSettingsLayout::OnBuildAdjacencyBufferChanged)
+		];
+	}
+
 	{
 		ChildrenBuilder.AddChildContent( LOCTEXT("UseFullPrecisionUVs", "Use Full Precision UVs") )
 		.NameContent()
@@ -613,6 +629,11 @@ ECheckBoxState FMeshBuildSettingsLayout::ShouldRemoveDegenerates() const
 	return BuildSettings.bRemoveDegenerates ? ECheckBoxState::Checked : ECheckBoxState::Unchecked;
 }
 
+ECheckBoxState FMeshBuildSettingsLayout::ShouldBuildAdjacencyBuffer() const
+{
+	return BuildSettings.bBuildAdjacencyBuffer ? ECheckBoxState::Checked : ECheckBoxState::Unchecked;
+}
+
 ECheckBoxState FMeshBuildSettingsLayout::ShouldUseFullPrecisionUVs() const
 {
 	return BuildSettings.bUseFullPrecisionUVs ? ECheckBoxState::Checked : ECheckBoxState::Unchecked;
@@ -711,6 +732,19 @@ void FMeshBuildSettingsLayout::OnRemoveDegeneratesChanged(ECheckBoxState NewStat
 	}
 }
 
+void FMeshBuildSettingsLayout::OnBuildAdjacencyBufferChanged(ECheckBoxState NewState)
+{
+	const bool bBuildAdjacencyBuffer = (NewState == ECheckBoxState::Checked) ? true : false;
+	if (BuildSettings.bBuildAdjacencyBuffer != bBuildAdjacencyBuffer)
+	{
+		if (FEngineAnalytics::IsAvailable())
+		{
+			FEngineAnalytics::GetProvider().RecordEvent(TEXT("Editor.Usage.StaticMesh.BuildSettings"), TEXT("bBuildAdjacencyBuffer"), bBuildAdjacencyBuffer ? TEXT("True") : TEXT("False"));
+		}
+		BuildSettings.bBuildAdjacencyBuffer = bBuildAdjacencyBuffer;
+	}
+}
+
 void FMeshBuildSettingsLayout::OnUseFullPrecisionUVsChanged(ECheckBoxState NewState)
 {
 	const bool bUseFullPrecisionUVs = (NewState == ECheckBoxState::Checked) ? true : false;
diff --git a/Engine/Source/Editor/StaticMeshEditor/Private/StaticMeshEditorTools.h b/Engine/Source/Editor/StaticMeshEditor/Private/StaticMeshEditorTools.h
--- a/Engine/Source/Editor/StaticMeshEditor/Private/StaticMeshEditorTools.h
+++ b/Engine/Source/Editor/StaticMeshEditor/Private/StaticMeshEditorTools.h
@@ -139,6 +139,7 @@ private:
 	ECheckBoxState ShouldRecomputeTangents() const;
 	ECheckBoxState ShouldUseMikkTSpace() const;
 	ECheckBoxState ShouldRemoveDegenerates() const;
+	ECheckBoxState ShouldBuildAdjacencyBuffer() const;
 	ECheckBoxState ShouldUseFullPrecisionUVs() const;
 	ECheckBoxState ShouldGenerateLightmapUVs() const;
 	ECheckBoxState ShouldGenerateDistanceFieldAsIfTwoSided() const;
@@ -154,6 +155,7 @@ private:
 	void OnRecomputeTangentsChanged(ECheckBoxState NewState);
 	void OnUseMikkTSpaceChanged(ECheckBoxState NewState);
 	void OnRemoveDegeneratesChanged(ECheckBoxState NewState);
+	void OnBuildAdjacencyBufferChanged(ECheckBoxState NewState);
 	void OnUseFullPrecisionUVsChanged(ECheckBoxState NewState);
 	void OnGenerateLightmapUVsChanged(ECheckBoxState NewState);
 	void OnGenerateDistanceFieldAsIfTwoSidedChanged(ECheckBoxState NewState);
diff --git a/Engine/Source/Editor/UnrealEd/Classes/Factories/FbxStaticMeshImportData.h b/Engine/Source/Editor/UnrealEd/Classes/Factories/FbxStaticMeshImportData.h
--- a/Engine/Source/Editor/UnrealEd/Classes/Factories/FbxStaticMeshImportData.h
+++ b/Engine/Source/Editor/UnrealEd/Classes/Factories/FbxStaticMeshImportData.h
@@ -41,6 +41,10 @@ class UFbxStaticMeshImportData : public UFbxMeshImportData
 	/** Disabling this option will keep degenerate triangles found.  In general you should leave this option on. */
 	UPROPERTY(EditAnywhere, AdvancedDisplay, Category = ImportSettings, meta = (ImportType = "StaticMesh"))
 	uint32 bRemoveDegenerates:1;
+
+	/** Required for PNT tessellation but can be slow. Recommend disabling for larger meshes. */
+	UPROPERTY(EditAnywhere, AdvancedDisplay, Category = ImportSettings, meta = (ImportType = "StaticMesh"))
+	uint32 bBuildAdjacencyBuffer:1;
 	
 	UPROPERTY(EditAnywhere, config, AdvancedDisplay, Category=ImportSettings, meta=(ImportType="StaticMesh"))
 	uint32 bGenerateLightmapUVs:1;
diff --git a/Engine/Source/Editor/UnrealEd/Private/Fbx/FbxMainImport.cpp b/Engine/Source/Editor/UnrealEd/Private/Fbx/FbxMainImport.cpp
--- a/Engine/Source/Editor/UnrealEd/Private/Fbx/FbxMainImport.cpp
+++ b/Engine/Source/Editor/UnrealEd/Private/Fbx/FbxMainImport.cpp
@@ -232,6 +232,7 @@ void ApplyImportUIToImportOptions(UFbxImportUI* ImportUI, FBXImportOptions& InOu
 	InOutImportOptions.VertexColorImportOption = ImportUI->StaticMeshImportData->VertexColorImportOption;
 	InOutImportOptions.VertexOverrideColor = ImportUI->StaticMeshImportData->VertexOverrideColor;
 	InOutImportOptions.bRemoveDegenerates = ImportUI->StaticMeshImportData->bRemoveDegenerates;
+	InOutImportOptions.bBuildAdjacencyBuffer = ImportUI->StaticMeshImportData->bBuildAdjacencyBuffer;
 	InOutImportOptions.bGenerateLightmapUVs = ImportUI->StaticMeshImportData->bGenerateLightmapUVs;
 	InOutImportOptions.bOneConvexHullPerUCX = ImportUI->StaticMeshImportData->bOneConvexHullPerUCX;
 	InOutImportOptions.bAutoGenerateCollision = ImportUI->StaticMeshImportData->bAutoGenerateCollision;
diff --git a/Engine/Source/Editor/UnrealEd/Private/Fbx/FbxStaticMeshImport.cpp b/Engine/Source/Editor/UnrealEd/Private/Fbx/FbxStaticMeshImport.cpp
--- a/Engine/Source/Editor/UnrealEd/Private/Fbx/FbxStaticMeshImport.cpp
+++ b/Engine/Source/Editor/UnrealEd/Private/Fbx/FbxStaticMeshImport.cpp
@@ -1208,6 +1208,7 @@ UStaticMesh* UnFbx::FFbxImporter::ImportStaticMeshAsSingle(UObject* InParent, TA
 
 		// @todo This overrides restored values currently but we need to be able to import over the existing settings if the user chose to do so.
 		SrcModel.BuildSettings.bRemoveDegenerates = ImportOptions->bRemoveDegenerates;
+		SrcModel.BuildSettings.bBuildAdjacencyBuffer = ImportOptions->bBuildAdjacencyBuffer;
 		SrcModel.BuildSettings.bRecomputeNormals = ImportOptions->NormalImportMethod == FBXNIM_ComputeNormals;
 		SrcModel.BuildSettings.bRecomputeTangents = ImportOptions->NormalImportMethod != FBXNIM_ImportNormalsAndTangents;
 		SrcModel.BuildSettings.bUseMikkTSpace = ImportOptions->NormalGenerationMethod == EFBXNormalGenerationMethod::MikkTSpace;
diff --git a/Engine/Source/Editor/UnrealEd/Private/Fbx/FbxStaticMeshImportData.cpp b/Engine/Source/Editor/UnrealEd/Private/Fbx/FbxStaticMeshImportData.cpp
--- a/Engine/Source/Editor/UnrealEd/Private/Fbx/FbxStaticMeshImportData.cpp
+++ b/Engine/Source/Editor/UnrealEd/Private/Fbx/FbxStaticMeshImportData.cpp
@@ -8,6 +8,7 @@ UFbxStaticMeshImportData::UFbxStaticMeshImportData(const FObjectInitializer& Obj
 {
 	StaticMeshLODGroup = NAME_None;
 	bRemoveDegenerates = true;
+	bBuildAdjacencyBuffer = true;
 	bGenerateLightmapUVs = true;
 	bOneConvexHullPerUCX = true;
 	bAutoGenerateCollision = true;
diff --git a/Engine/Source/Editor/UnrealEd/Private/FbxImporter.h b/Engine/Source/Editor/UnrealEd/Private/FbxImporter.h
--- a/Engine/Source/Editor/UnrealEd/Private/FbxImporter.h
+++ b/Engine/Source/Editor/UnrealEd/Private/FbxImporter.h
@@ -124,6 +124,7 @@ struct FBXImportOptions
 	EVertexColorImportOption::Type VertexColorImportOption;
 	FColor VertexOverrideColor;
 	bool bRemoveDegenerates;
+	bool bBuildAdjacencyBuffer;
 	bool bGenerateLightmapUVs;
 	bool bOneConvexHullPerUCX;
 	bool bAutoGenerateCollision;
diff --git a/Engine/Source/Runtime/Core/Public/UObject/ObjectVersion.h b/Engine/Source/Runtime/Core/Public/UObject/ObjectVersion.h
--- a/Engine/Source/Runtime/Core/Public/UObject/ObjectVersion.h
+++ b/Engine/Source/Runtime/Core/Public/UObject/ObjectVersion.h
@@ -554,8 +554,8 @@ enum EUnrealEngineObjectUE4Version
 	VER_UE4_NAVIGATION_AGENT_SELECTOR,
 	// Mesh particle collisions consider particle size.
 	VER_UE4_MESH_PARTICLE_COLLISIONS_CONSIDER_PARTICLE_SIZE,
-
-
+	// Adjacency buffer building no longer automatically handled based on triangle count, user-controlled
+	VER_UE4_BUILD_MESH_ADJ_BUFFER_FLAG_EXPOSED,
 	// -----<new versions can be added before this line>-------------------------------------------------
 	// - this needs to be the last line (see note below)
 	VER_UE4_AUTOMATIC_VERSION_PLUS_ONE,
diff --git a/Engine/Source/Runtime/Engine/Classes/Engine/EngineTypes.h b/Engine/Source/Runtime/Engine/Classes/Engine/EngineTypes.h
--- a/Engine/Source/Runtime/Engine/Classes/Engine/EngineTypes.h
+++ b/Engine/Source/Runtime/Engine/Classes/Engine/EngineTypes.h
@@ -2145,6 +2145,10 @@ struct FMeshBuildSettings
 	/** If true, degenerate triangles will be removed. */
 	UPROPERTY(EditAnywhere, Category=BuildSettings)
 	bool bRemoveDegenerates;
+	
+	/** Required for PNT tessellation but can be slow. Recommend disabling for larger meshes. */
+	UPROPERTY(EditAnywhere, Category=BuildSettings)
+	bool bBuildAdjacencyBuffer;
 
 	/** If true, UVs will be stored at full floating point precision. */
 	UPROPERTY(EditAnywhere, Category=BuildSettings)
@@ -2192,6 +2196,7 @@ struct FMeshBuildSettings
 		, bRecomputeNormals(true)
 		, bRecomputeTangents(true)
 		, bRemoveDegenerates(true)
+		, bBuildAdjacencyBuffer(true)
 		, bUseFullPrecisionUVs(false)
 		, bGenerateLightmapUVs(true)
 		, MinLightmapResolution(64)
@@ -2211,6 +2216,7 @@ struct FMeshBuildSettings
 			&& bRecomputeTangents == Other.bRecomputeTangents
 			&& bUseMikkTSpace == Other.bUseMikkTSpace
 			&& bRemoveDegenerates == Other.bRemoveDegenerates
+			&& bBuildAdjacencyBuffer == Other.bBuildAdjacencyBuffer
 			&& bUseFullPrecisionUVs == Other.bUseFullPrecisionUVs
 			&& bGenerateLightmapUVs == Other.bGenerateLightmapUVs
 			&& MinLightmapResolution == Other.MinLightmapResolution
diff --git a/Engine/Source/Runtime/Engine/Private/StaticMesh.cpp b/Engine/Source/Runtime/Engine/Private/StaticMesh.cpp
--- a/Engine/Source/Runtime/Engine/Private/StaticMesh.cpp
+++ b/Engine/Source/Runtime/Engine/Private/StaticMesh.cpp
@@ -1024,6 +1024,7 @@ FArchive& operator<<(FArchive& Ar, FMeshBuildSettings& BuildSettings)
 	Ar << BuildSettings.bRecomputeTangents;
 	Ar << BuildSettings.bUseMikkTSpace;
 	Ar << BuildSettings.bRemoveDegenerates;
+	Ar << BuildSettings.bBuildAdjacencyBuffer;
 	Ar << BuildSettings.bUseFullPrecisionUVs;
 	Ar << BuildSettings.bGenerateLightmapUVs;
 
@@ -1958,6 +1959,27 @@ void UStaticMesh::PostLoad()
 		}
 	}
 
+	if (GetLinkerUE4Version() < VER_UE4_BUILD_MESH_ADJ_BUFFER_FLAG_EXPOSED)
+	{
+		FRawMesh TempRawMesh;
+		uint32 TotalIndexCount = 0;
+
+		for (int32 i = 0; i < SourceModels.Num(); ++i)
+		{
+			FRawMeshBulkData* RawMeshBulkData = SourceModels[i].RawMeshBulkData;
+			if (RawMeshBulkData)
+			{
+				RawMeshBulkData->LoadRawMesh(TempRawMesh);
+				TotalIndexCount += TempRawMesh.WedgeIndices.Num();
+			}
+		}
+
+		for (int32 i = 0; i < SourceModels.Num(); ++i)
+		{
+			SourceModels[i].BuildSettings.bBuildAdjacencyBuffer = (TotalIndexCount < 50000);
+		}
+	}
+
 	CacheDerivedData();
 
 	// Only required in an editor build as other builds process this in a different place
