Cancelling at the prompt to change LOD Group in the Static Mesh Editor now no longer changes the LOD group.

diff --git a/Engine/Source/Editor/StaticMeshEditor/Private/StaticMeshEditorTools.cpp b/Engine/Source/Editor/StaticMeshEditor/Private/StaticMeshEditorTools.cpp
--- a/Engine/Source/Editor/StaticMeshEditor/Private/StaticMeshEditorTools.cpp
+++ b/Engine/Source/Editor/StaticMeshEditor/Private/StaticMeshEditorTools.cpp
@@ -1508,7 +1508,7 @@ void FLevelOfDetailSettingsLayout::AddToDetailsPanel( IDetailLayoutBuilder& Deta
 	]
 	.ValueContent()
 	[
-		SNew(STextComboBox)
+		SAssignNew(LODGroupComboBox, STextComboBox)
 		.ContentPadding(0)
 		.OptionsSource(&LODGroupOptions)
 		.InitiallySelectedItem(LODGroupOptions[(LODGroupIndex == INDEX_NONE) ? 0 : LODGroupIndex])
@@ -1883,14 +1883,15 @@ void FLevelOfDetailSettingsLayout::OnLODGroupChanged(TSharedPtr<FString> NewValu
 	FName NewGroup = LODGroupNames[GroupIndex];
 	if (StaticMesh->LODGroup != NewGroup)
 	{
-		StaticMesh->Modify();
-		StaticMesh->LODGroup = NewGroup;
 		EAppReturnType::Type DialogResult = FMessageDialog::Open(
 			EAppMsgType::YesNo,
-			FText::Format( LOCTEXT("ApplyDefaultLODSettings", "Overwrite settings with the defaults from LOD group '{0}'?"), FText::FromString( **NewValue ) )
+			FText::Format( LOCTEXT("ApplyDefaultLODSettings", "Changing LOD group will overwrite the current settings with the defaults from LOD group '{0}'. Do you wish to continue?"), FText::FromString( **NewValue ) )
 			);
 		if (DialogResult == EAppReturnType::Yes)
 		{
+			StaticMesh->Modify();
+			StaticMesh->LODGroup = NewGroup;
+
 			const ITargetPlatform* Platform = GetTargetPlatformManagerRef().GetRunningTargetPlatform();
 			check(Platform);
 			const FStaticMeshLODGroup& GroupSettings = Platform->GetStaticMeshLODSettings().GetLODGroup(NewGroup);
@@ -1911,9 +1912,16 @@ void FLevelOfDetailSettingsLayout::OnLODGroupChanged(TSharedPtr<FString> NewValu
 			}
 			StaticMesh->bAutoComputeLODScreenSize = true;
 			StaticMesh->LightMapResolution = GroupSettings.GetDefaultLightMapResolution();
+			StaticMesh->PostEditChange();
+			StaticMeshEditor.RefreshTool();
+		}
+		else
+		{
+			// Overriding the selection; ensure that the widget correctly reflects the property value
+			int32 Index = LODGroupNames.Find(StaticMesh->LODGroup);
+			check(Index != INDEX_NONE);
+			LODGroupComboBox->SetSelectedItem(LODGroupOptions[Index]);
 		}
-		StaticMesh->PostEditChange();
-		StaticMeshEditor.RefreshTool();
 	}
 }
 
diff --git a/Engine/Source/Editor/StaticMeshEditor/Private/StaticMeshEditorTools.h b/Engine/Source/Editor/StaticMeshEditor/Private/StaticMeshEditorTools.h
--- a/Engine/Source/Editor/StaticMeshEditor/Private/StaticMeshEditorTools.h
+++ b/Engine/Source/Editor/StaticMeshEditor/Private/StaticMeshEditorTools.h
@@ -328,6 +328,9 @@ private:
 	TSharedPtr<FMeshBuildSettingsLayout> BuildSettingsWidgets[MAX_STATIC_MESH_LODS];
 	TSharedPtr<FMeshSectionSettingsLayout> SectionSettingsWidgets[MAX_STATIC_MESH_LODS];
 
+	/** ComboBox widget for the LOD Group property */
+	TSharedPtr<STextComboBox> LODGroupComboBox;
+
 	/** The display factors at which LODs swap */
 	float LODScreenSizes[MAX_STATIC_MESH_LODS];
 
