'Num Primitives' -> 'Num Collision Primitives'

diff --git a/Engine/Source/Editor/StaticMeshEditor/Private/StaticMeshEditorViewportClient.cpp b/Engine/Source/Editor/StaticMeshEditor/Private/StaticMeshEditorViewportClient.cpp
--- a/Engine/Source/Editor/StaticMeshEditor/Private/StaticMeshEditorViewportClient.cpp
+++ b/Engine/Source/Editor/StaticMeshEditor/Private/StaticMeshEditorViewportClient.cpp
@@ -795,7 +795,7 @@ void FStaticMeshEditorViewportClient::DrawCanvas( FViewport& InViewport, FSceneV
 	if(bShowCollision && StaticMesh->BodySetup)
 	{
 		TextItems.Add(SStaticMeshEditorViewport::FOverlayTextItem(
-			FText::Format(NSLOCTEXT("UnrealEd", "NumPrimitives_F", "Num Primitives:  {0}"), FText::AsNumber(StaticMesh->BodySetup->AggGeom.GetElementCount()))));
+			FText::Format(NSLOCTEXT("UnrealEd", "NumPrimitives_F", "Num Collision Primitives:  {0}"), FText::AsNumber(StaticMesh->BodySetup->AggGeom.GetElementCount()))));
 	}
 
 	StaticMeshEditorViewport->PopulateOverlayText(TextItems);
