[UE-19897] Continued. Added Static Mesh support. Sections can now be isolated in the Static Mesh Editor.

diff --git a/Engine/Source/Editor/StaticMeshEditor/Private/StaticMeshEditorTools.cpp b/Engine/Source/Editor/StaticMeshEditor/Private/StaticMeshEditorTools.cpp
--- a/Engine/Source/Editor/StaticMeshEditor/Private/StaticMeshEditorTools.cpp
+++ b/Engine/Source/Editor/StaticMeshEditor/Private/StaticMeshEditorTools.cpp
@@ -1,4 +1,4 @@
-// Copyright 1998-2015 Epic Games, Inc. All Rights Reserved.
+		// Copyright 1998-2015 Epic Games, Inc. All Rights Reserved.
 
 #include "StaticMeshEditorModule.h"
 #include "StaticMeshEditorTools.h"
@@ -1274,17 +1274,37 @@ void FMeshSectionSettingsLayout::OnMaterialChanged(UMaterialInterface* NewMateri
 
 TSharedRef<SWidget> FMeshSectionSettingsLayout::OnGenerateNameWidgetsForMaterial(UMaterialInterface* Material, int32 SlotIndex)
 {
-	return
-		SNew(SCheckBox)
-		.IsChecked(this, &FMeshSectionSettingsLayout::IsSectionSelected, SlotIndex)
-		.OnCheckStateChanged(this, &FMeshSectionSettingsLayout::OnSectionSelectedChanged, SlotIndex)
-		.ToolTipText(LOCTEXT("Highlight_ToolTip", "Highlights this section in the viewport"))
+	return SNew(SVerticalBox)
+		+ SVerticalBox::Slot()
+		.AutoHeight()
 		[
-			SNew(STextBlock)
-			.Font(IDetailLayoutBuilder::GetDetailFont())
-			.ColorAndOpacity( FLinearColor( 0.4f, 0.4f, 0.4f, 1.0f) )
-			.Text(LOCTEXT("Highlight", "Highlight"))
+			SNew(SCheckBox)
+			.IsChecked(this, &FMeshSectionSettingsLayout::IsSectionHighlighted, SlotIndex)
+			.OnCheckStateChanged(this, &FMeshSectionSettingsLayout::OnSectionHighlightedChanged, SlotIndex)
+			.ToolTipText(LOCTEXT("Highlight_ToolTip", "Highlights this section in the viewport"))
+			[
+				SNew(STextBlock)
+				.Font(IDetailLayoutBuilder::GetDetailFont())
+				.ColorAndOpacity( FLinearColor( 0.4f, 0.4f, 0.4f, 1.0f) )
+				.Text(LOCTEXT("Highlight", "Highlight"))
 
+			]
+		]
+		+ SVerticalBox::Slot()
+		.AutoHeight()
+		.Padding(0, 2, 0, 0)
+		[
+			SNew(SCheckBox)
+			.IsChecked(this, &FMeshSectionSettingsLayout::IsSectionIsolatedEnabled, SlotIndex)
+			.OnCheckStateChanged(this, &FMeshSectionSettingsLayout::OnSectionIsolatedChanged, SlotIndex)
+			.ToolTipText(LOCTEXT("Isolate_ToolTip", "Isolates this section in the viewport"))
+			[
+				SNew(STextBlock)
+				.Font(IDetailLayoutBuilder::GetDetailFont())
+				.ColorAndOpacity(FLinearColor(0.4f, 0.4f, 0.4f, 1.0f))
+				.Text(LOCTEXT("Isolate", "Isolate"))
+
+			]
 		];
 }
 
@@ -1366,7 +1386,7 @@ void FMeshSectionSettingsLayout::OnSectionCollisionChanged(ECheckBoxState NewSta
 	CallPostEditChange();
 }
 
-ECheckBoxState FMeshSectionSettingsLayout::IsSectionSelected(int32 SectionIndex) const
+ECheckBoxState FMeshSectionSettingsLayout::IsSectionHighlighted(int32 SectionIndex) const
 {
 	ECheckBoxState State = ECheckBoxState::Unchecked;
 	UStaticMeshComponent* Component = StaticMeshEditor.GetStaticMeshComponent();
@@ -1377,7 +1397,7 @@ ECheckBoxState FMeshSectionSettingsLayout::IsSectionSelected(int32 SectionIndex)
 	return State;
 }
 
-void FMeshSectionSettingsLayout::OnSectionSelectedChanged(ECheckBoxState NewState, int32 SectionIndex)
+void FMeshSectionSettingsLayout::OnSectionHighlightedChanged(ECheckBoxState NewState, int32 SectionIndex)
 {
 	UStaticMeshComponent* Component = StaticMeshEditor.GetStaticMeshComponent();
 	if (Component)
@@ -1385,6 +1405,11 @@ void FMeshSectionSettingsLayout::OnSectionSelectedChanged(ECheckBoxState NewStat
 		if (NewState == ECheckBoxState::Checked)
 		{
 			Component->SelectedEditorSection = SectionIndex;
+			if (Component->SectionIndexPreview != SectionIndex)
+			{
+				// Unhide all mesh sections
+				Component->SetSectionPreview(INDEX_NONE);
+			}
 		}
 		else if (NewState == ECheckBoxState::Unchecked)
 		{
@@ -1395,6 +1420,39 @@ void FMeshSectionSettingsLayout::OnSectionSelectedChanged(ECheckBoxState NewStat
 	}
 }
 
+ECheckBoxState FMeshSectionSettingsLayout::IsSectionIsolatedEnabled(int32 SectionIndex) const
+{
+	ECheckBoxState State = ECheckBoxState::Unchecked;
+	const UStaticMeshComponent* Component = StaticMeshEditor.GetStaticMeshComponent();
+	if (Component)
+	{
+		State = Component->SectionIndexPreview == SectionIndex ? ECheckBoxState::Checked : ECheckBoxState::Unchecked;
+	}
+	return State;
+}
+
+void FMeshSectionSettingsLayout::OnSectionIsolatedChanged(ECheckBoxState NewState, int32 SectionIndex)
+{
+	UStaticMeshComponent* Component = StaticMeshEditor.GetStaticMeshComponent();
+	if (Component)
+	{
+		if (NewState == ECheckBoxState::Checked)
+		{
+			Component->SetSectionPreview(SectionIndex);
+			if (Component->SelectedEditorSection != SectionIndex)
+			{
+				Component->SelectedEditorSection = INDEX_NONE;
+			}
+		}
+		else if (NewState == ECheckBoxState::Unchecked)
+		{
+			Component->SetSectionPreview(INDEX_NONE);
+		}
+		Component->MarkRenderStateDirty();
+		StaticMeshEditor.RefreshViewport();
+	}
+}
+
 void FMeshSectionSettingsLayout::CallPostEditChange(UProperty* PropertyChanged/*=nullptr*/)
 {
 	UStaticMesh& StaticMesh = GetStaticMesh();
diff --git a/Engine/Source/Editor/StaticMeshEditor/Private/StaticMeshEditorTools.h b/Engine/Source/Editor/StaticMeshEditor/Private/StaticMeshEditorTools.h
--- a/Engine/Source/Editor/StaticMeshEditor/Private/StaticMeshEditorTools.h
+++ b/Engine/Source/Editor/StaticMeshEditor/Private/StaticMeshEditorTools.h
@@ -242,8 +242,10 @@ private:
 	void OnSectionCastShadowChanged(ECheckBoxState NewState, int32 SectionIndex);
 	ECheckBoxState DoesSectionCollide(int32 SectionIndex) const;
 	void OnSectionCollisionChanged(ECheckBoxState NewState, int32 SectionIndex);
-	ECheckBoxState IsSectionSelected(int32 SectionIndex) const;
-	void OnSectionSelectedChanged(ECheckBoxState NewState, int32 SectionIndex);
+	ECheckBoxState IsSectionHighlighted(int32 SectionIndex) const;
+	void OnSectionHighlightedChanged(ECheckBoxState NewState, int32 SectionIndex);
+	ECheckBoxState IsSectionIsolatedEnabled(int32 SectionIndex) const;
+	void OnSectionIsolatedChanged(ECheckBoxState NewState, int32 SectionIndex);
 	void CallPostEditChange(UProperty* PropertyChanged=nullptr);
 	
 	IStaticMeshEditor& StaticMeshEditor;
diff --git a/Engine/Source/Editor/StaticMeshEditor/Private/StaticMeshEditorViewportClient.cpp b/Engine/Source/Editor/StaticMeshEditor/Private/StaticMeshEditorViewportClient.cpp
--- a/Engine/Source/Editor/StaticMeshEditor/Private/StaticMeshEditorViewportClient.cpp
+++ b/Engine/Source/Editor/StaticMeshEditor/Private/StaticMeshEditorViewportClient.cpp
@@ -798,6 +798,11 @@ void FStaticMeshEditorViewportClient::DrawCanvas( FViewport& InViewport, FSceneV
 			FText::Format(NSLOCTEXT("UnrealEd", "NumPrimitives_F", "Num Collision Primitives:  {0}"), FText::AsNumber(StaticMesh->BodySetup->AggGeom.GetElementCount()))));
 	}
 
+	if (StaticMeshComponent && StaticMeshComponent->SectionIndexPreview != INDEX_NONE)
+	{
+		TextItems.Add(SStaticMeshEditorViewport::FOverlayTextItem(NSLOCTEXT("UnrealEd", "MeshSectionsHiddenWarning",  "Mesh Sections Hidden")));
+	}
+
 	StaticMeshEditorViewport->PopulateOverlayText(TextItems);
 
 	if(bDrawUVs)
diff --git a/Engine/Source/Runtime/Engine/Classes/Components/StaticMeshComponent.h b/Engine/Source/Runtime/Engine/Classes/Components/StaticMeshComponent.h
--- a/Engine/Source/Runtime/Engine/Classes/Components/StaticMeshComponent.h
+++ b/Engine/Source/Runtime/Engine/Classes/Components/StaticMeshComponent.h
@@ -156,9 +156,12 @@ class ENGINE_API UStaticMeshComponent : public UMeshComponent
 	FColor WireframeColorOverride;
 
 #if WITH_EDITORONLY_DATA
-	/** The section currently selected in the Editor. */
+	/** The section currently selected in the Editor. Used for highlighting */
 	UPROPERTY(transient)
 	int32 SelectedEditorSection;
+	/** Index of the section to preview. If set to INDEX_NONE, all section will be rendered. Used for isolating in Static Mesh Tool **/
+	UPROPERTY(transient)
+	int32 SectionIndexPreview;
 #endif
 
 	/** If true, bForceNavigationObstacle flag will take priority over navigation data stored in StaticMesh */
@@ -409,6 +412,13 @@ public:
 	 */
 	void RemoveInstanceVertexColors();
 
+	/**
+	*	Sets the value of the SectionIndexPreview flag and reattaches the component as necessary.
+	*	@param	InSectionIndexPreview		New value of SectionIndexPreview.
+	*/
+	void SetSectionPreview(int32 InSectionIndexPreview);
+
+
 private:
 	/** Initializes the resources used by the static mesh component. */
 	void InitResources();
diff --git a/Engine/Source/Runtime/Engine/Private/Components/StaticMeshComponent.cpp b/Engine/Source/Runtime/Engine/Private/Components/StaticMeshComponent.cpp
--- a/Engine/Source/Runtime/Engine/Private/Components/StaticMeshComponent.cpp
+++ b/Engine/Source/Runtime/Engine/Private/Components/StaticMeshComponent.cpp
@@ -185,6 +185,7 @@ UStaticMeshComponent::UStaticMeshComponent(const FObjectInitializer& ObjectIniti
 
 #if WITH_EDITORONLY_DATA
 	SelectedEditorSection = INDEX_NONE;
+	SectionIndexPreview = INDEX_NONE;
 #endif
 }
 
@@ -722,6 +723,18 @@ bool UStaticMeshComponent::RequiresOverrideVertexColorsFixup( TArray<int32>& Out
 	return bFixupRequired;
 }
 
+void UStaticMeshComponent::SetSectionPreview(int32 InSectionIndexPreview)
+{
+#if WITH_EDITORONLY_DATA
+	if (SectionIndexPreview != InSectionIndexPreview)
+	{
+		SectionIndexPreview = InSectionIndexPreview;
+		MarkRenderStateDirty();
+	}
+#endif
+}
+
+
 void UStaticMeshComponent::RemoveInstanceVertexColorsFromLOD( int32 LODToRemoveColorsFrom )
 {
 #if WITH_EDITORONLY_DATA
diff --git a/Engine/Source/Runtime/Engine/Private/StaticMeshRender.cpp b/Engine/Source/Runtime/Engine/Private/StaticMeshRender.cpp
--- a/Engine/Source/Runtime/Engine/Private/StaticMeshRender.cpp
+++ b/Engine/Source/Runtime/Engine/Private/StaticMeshRender.cpp
@@ -52,16 +52,19 @@ static FAutoConsoleCommand GToggleForceDefaultMaterialCmd(
 
 /** Initialization constructor. */
 FStaticMeshSceneProxy::FStaticMeshSceneProxy(UStaticMeshComponent* InComponent):
-	FPrimitiveSceneProxy(InComponent, InComponent->StaticMesh->GetFName()),
-	Owner(InComponent->GetOwner()),
-	StaticMesh(InComponent->StaticMesh),
-	BodySetup(InComponent->GetBodySetup()),
-	RenderData(InComponent->StaticMesh->RenderData),
-	ForcedLodModel(InComponent->ForcedLodModel),
-	bCastShadow(InComponent->CastShadow),
-	CollisionTraceFlag(ECollisionTraceFlag::CTF_UseDefault),
-	MaterialRelevance(InComponent->GetMaterialRelevance(GetScene().GetFeatureLevel())),
-	CollisionResponse(InComponent->GetCollisionResponseToChannels())
+	FPrimitiveSceneProxy(InComponent, InComponent->StaticMesh->GetFName())
+	, Owner(InComponent->GetOwner())
+	, StaticMesh(InComponent->StaticMesh)
+	, BodySetup(InComponent->GetBodySetup())
+	, RenderData(InComponent->StaticMesh->RenderData)
+	, ForcedLodModel(InComponent->ForcedLodModel)
+	, bCastShadow(InComponent->CastShadow)
+	, CollisionTraceFlag(ECollisionTraceFlag::CTF_UseDefault)
+	, MaterialRelevance(InComponent->GetMaterialRelevance(GetScene().GetFeatureLevel()))
+	, CollisionResponse(InComponent->GetCollisionResponseToChannels())
+#if WITH_EDITORONLY_DATA
+	, SectionIndexPreview(InComponent->SectionIndexPreview)
+#endif
 {
 	check(RenderData);
 
@@ -224,6 +227,14 @@ bool FStaticMeshSceneProxy::GetMeshElement(int32 LODIndex, int32 BatchIndex, int
 	OutMeshBatch.MaterialRenderProxy = Material->GetRenderProxy(bUseSelectedMaterial,bUseHoveredMaterial);
 	OutMeshBatch.VertexFactory = &LOD.VertexFactory;
 
+#if WITH_EDITORONLY_DATA
+	// If section is hidden, then skip the draw.
+	if ((SectionIndexPreview >= 0) && (SectionIndexPreview != SectionIndex))
+	{
+		return false;
+	}
+#endif
+
 	// Has the mesh component overridden the vertex color stream for this mesh LOD?
 	if( ProxyLODInfo.OverrideColorVertexBuffer != NULL )
 	{
diff --git a/Engine/Source/Runtime/Engine/Public/StaticMeshResources.h b/Engine/Source/Runtime/Engine/Public/StaticMeshResources.h
--- a/Engine/Source/Runtime/Engine/Public/StaticMeshResources.h
+++ b/Engine/Source/Runtime/Engine/Public/StaticMeshResources.h
@@ -869,6 +869,11 @@ protected:
 
 	const FDistanceFieldVolumeData* DistanceFieldData;
 
+#if WITH_EDITORONLY_DATA
+	/** Index of the section to preview. If set to INDEX_NONE, all section will be rendered */
+	int32 SectionIndexPreview;
+#endif
+
 	/**
 	 * The forcedLOD set in the static mesh editor, copied from the mesh component
 	 */
