Fix to "Focus Viewport to Selection" in Static Mesh Editor.

diff --git a/Engine/Source/Editor/StaticMeshEditor/Private/StaticMeshEditorViewportClient.cpp b/Engine/Source/Editor/StaticMeshEditor/Private/StaticMeshEditorViewportClient.cpp
--- a/Engine/Source/Editor/StaticMeshEditor/Private/StaticMeshEditorViewportClient.cpp
+++ b/Engine/Source/Editor/StaticMeshEditor/Private/StaticMeshEditorViewportClient.cpp
@@ -1174,8 +1174,11 @@ void FStaticMeshEditorViewportClient::PerspectiveCameraMoved()
 		);
 
 	ToggleOrbitCamera(bWasOrbit);
-	SetViewLocation(OldCameraLocation);
-	SetViewRotation(OldCameraRotation);
+	if (!GetViewTransform().IsPlaying())
+	{
+		SetViewLocation(OldCameraLocation);
+		SetViewRotation(OldCameraRotation);
+	}
 }
 
 void FStaticMeshEditorViewportClient::SetPreviewMesh(UStaticMesh* InStaticMesh, UStaticMeshComponent* InStaticMeshComponent)
