Copying //UE4/Dev-Editor to //UE4/Dev-Main (Source: //UE4/Dev-Editor @ 3228984)

diff --git a/Engine/Source/Editor/PropertyEditor/Private/UserInterface/PropertyEditor/SPropertyEditorCombo.cpp b/Engine/Source/Editor/PropertyEditor/Private/UserInterface/PropertyEditor/SPropertyEditorCombo.cpp
--- a/Engine/Source/Editor/PropertyEditor/Private/UserInterface/PropertyEditor/SPropertyEditorCombo.cpp
+++ b/Engine/Source/Editor/PropertyEditor/Private/UserInterface/PropertyEditor/SPropertyEditorCombo.cpp
@@ -64,26 +64,69 @@
 {
 	PropertyEditor = InPropertyEditor;
 	
-	// Important to find certain info out about the combo box.
-	TArray< TSharedPtr<FString> > ComboItems;
-	TArray< FText > ToolTips;
-	TArray< bool > Restrictions;
-	GenerateComboBoxStrings( ComboItems, ToolTips, Restrictions );
-
+	TArray<TSharedPtr<FString>> ComboItems;
+	TArray<bool> Restrictions;
 	TArray<TSharedPtr<SToolTip>> RichToolTips;
 
+	GenerateComboBoxStrings(ComboItems, RichToolTips, Restrictions);
+
+	SAssignNew(ComboBox, SPropertyComboBox)
+		.Font( InArgs._Font )
+		.RichToolTipList( RichToolTips )
+		.ComboItemList( ComboItems )
+		.RestrictedList( Restrictions )
+		.OnSelectionChanged( this, &SPropertyEditorCombo::OnComboSelectionChanged )
+		.OnComboBoxOpening( this, &SPropertyEditorCombo::OnComboOpening )
+		.VisibleText( this, &SPropertyEditorCombo::GetDisplayValueAsString )
+		.ToolTipText( InPropertyEditor, &FPropertyEditor::GetValueAsText );
+
+	ChildSlot
+	[
+		ComboBox.ToSharedRef()
+	];
+
+	SetEnabled( TAttribute<bool>( this, &SPropertyEditorCombo::CanEdit ) );
+}
+
+FString SPropertyEditorCombo::GetDisplayValueAsString() const
+{
+	const UProperty* Property = PropertyEditor->GetProperty();
+	const bool bStringEnumProperty = Cast<UStrProperty>(Property) && Property->HasMetaData(TEXT("Enum"));
+
+	if ( !Cast<UByteProperty>(Property) && !Cast<UEnumProperty>(Property) && !bStringEnumProperty )
+	{
+		UObject* ObjectValue = NULL;
+		FPropertyAccess::Result Result = PropertyEditor->GetPropertyHandle()->GetValue( ObjectValue );
+
+		if( Result == FPropertyAccess::Success && ObjectValue != NULL )
+		{
+			return ObjectValue->GetName();
+		}
+	}
+
+	return (bUsesAlternateDisplayValues) ? PropertyEditor->GetValueAsDisplayString() : PropertyEditor->GetValueAsString();
+}
+
+
+void SPropertyEditorCombo::GenerateComboBoxStrings( TArray< TSharedPtr<FString> >& OutComboBoxStrings, TArray<TSharedPtr<SToolTip>>& RichToolTips, TArray<bool>& OutRestrictedItems )
+{
+	TArray<FText> BasicTooltips;
+	const TSharedRef< IPropertyHandle > PropertyHandle = PropertyEditor->GetPropertyHandle();
+	bUsesAlternateDisplayValues = PropertyHandle->GeneratePossibleValues(OutComboBoxStrings, BasicTooltips, OutRestrictedItems);
+
+
 	// For enums, look for rich tooltip information
 	if(PropertyEditor.IsValid())
 	{
 		if(const UProperty* Property = PropertyEditor->GetProperty())
 		{
 			UEnum* Enum = nullptr;
-			
+
 			if(const UByteProperty* ByteProperty = Cast<UByteProperty>(Property))
 			{
 				Enum = ByteProperty->Enum;
 			}
-			else if (const UEnumProperty* EnumProperty = Cast<UEnumProperty>(Property))
+			else if(const UEnumProperty* EnumProperty = Cast<UEnumProperty>(Property))
 			{
 				Enum = EnumProperty->GetEnum();
 			}
@@ -94,13 +137,13 @@
 
 				// Get enum doc link (not just GetDocumentationLink as that is the documentation for the struct we're in, not the enum documentation)
 				FString DocLink = PropertyEditorHelpers::GetEnumDocumentationLink(Property);
-			
+
 				for(int32 EnumIdx = 0; EnumIdx < Enum->NumEnums() - 1; ++EnumIdx)
 				{
 					FString Excerpt = Enum->GetEnumName(EnumIdx);
 
 					bool bShouldBeHidden = Enum->HasMetaData(TEXT("Hidden"), EnumIdx) || Enum->HasMetaData(TEXT("Spacer"), EnumIdx);
-					if (!bShouldBeHidden && AllowedPropertyEnums.Num() != 0)
+					if(!bShouldBeHidden && AllowedPropertyEnums.Num() != 0)
 					{
 						bShouldBeHidden = AllowedPropertyEnums.Find(Enum->GetEnum(EnumIdx)) == INDEX_NONE;
 					}
@@ -112,56 +155,14 @@
 						bShouldBeHidden = PropertyEditor->GetPropertyHandle()->IsHidden(EnumValueName);
 					}
 				
-					if (!bShouldBeHidden)
+					if(!bShouldBeHidden)
 					{
-						RichToolTips.Add(IDocumentation::Get()->CreateToolTip(MoveTemp(ToolTips[EnumIdx]), nullptr, DocLink, MoveTemp(Excerpt)));
+						RichToolTips.Add(IDocumentation::Get()->CreateToolTip(MoveTemp(BasicTooltips[EnumIdx]), nullptr, DocLink, MoveTemp(Excerpt)));
 					}
 				}
 			}
 		}
 	}
-
-	SAssignNew(ComboBox, SPropertyComboBox)
-		.Font( InArgs._Font )
-		.RichToolTipList( RichToolTips )
-		.ComboItemList( ComboItems )
-		.RestrictedList( Restrictions )
-		.OnSelectionChanged( this, &SPropertyEditorCombo::OnComboSelectionChanged )
-		.OnComboBoxOpening( this, &SPropertyEditorCombo::OnComboOpening )
-		.VisibleText( this, &SPropertyEditorCombo::GetDisplayValueAsString )
-		.ToolTipText( InPropertyEditor, &FPropertyEditor::GetValueAsText );
-
-	ChildSlot
-	[
-		ComboBox.ToSharedRef()
-	];
-
-	SetEnabled( TAttribute<bool>( this, &SPropertyEditorCombo::CanEdit ) );
 }
 
-FString SPropertyEditorCombo::GetDisplayValueAsString() const
-{
-	const UProperty* Property = PropertyEditor->GetProperty();
-	const bool bStringEnumProperty = Cast<UStrProperty>(Property) && Property->HasMetaData(TEXT("Enum"));
-
-	if ( !Cast<UByteProperty>(Property) && !Cast<UEnumProperty>(Property) && !bStringEnumProperty )
-	{
-		UObject* ObjectValue = NULL;
-		FPropertyAccess::Result Result = PropertyEditor->GetPropertyHandle()->GetValue( ObjectValue );
-
-		if( Result == FPropertyAccess::Success && ObjectValue != NULL )
-		{
-			return ObjectValue->GetName();
-		}
-	}
-
-	return (bUsesAlternateDisplayValues) ? PropertyEditor->GetValueAsDisplayString() : PropertyEditor->GetValueAsString();
-}
-
-void SPropertyEditorCombo::GenerateComboBoxStrings( TArray< TSharedPtr<FString> >& OutComboBoxStrings, TArray< FText >& OutToolTips, TArray<bool>& OutRestrictedItems )
-{
-	const TSharedRef< IPropertyHandle > PropertyHandle = PropertyEditor->GetPropertyHandle();
-	bUsesAlternateDisplayValues = PropertyHandle->GeneratePossibleValues(OutComboBoxStrings, OutToolTips, OutRestrictedItems);
-}
-
 void SPropertyEditorCombo::OnComboSelectionChanged( TSharedPtr<FString> NewValue, ESelectInfo::Type SelectInfo )
@@ -174,6 +175,13 @@
 
 void SPropertyEditorCombo::OnComboOpening()
 {
+	TArray<TSharedPtr<FString>> ComboItems;
+	TArray<TSharedPtr<SToolTip>> RichToolTips;
+	TArray<bool> Restrictions;
+	GenerateComboBoxStrings(ComboItems, RichToolTips, Restrictions);
+
+	ComboBox->SetItemList(ComboItems, RichToolTips, Restrictions);
+
 	// try and re-sync the selection in the combo list in case it was changed since Construct was called
 	// this would fail if the displayed value doesn't match the equivalent value in the combo list
 	FString CurrentDisplayValue = GetDisplayValueAsString();
diff --git a/Engine/Source/Editor/PropertyEditor/Private/UserInterface/PropertyEditor/SPropertyEditorCombo.h b/Engine/Source/Editor/PropertyEditor/Private/UserInterface/PropertyEditor/SPropertyEditorCombo.h
--- a/Engine/Source/Editor/PropertyEditor/Private/UserInterface/PropertyEditor/SPropertyEditorCombo.h
+++ b/Engine/Source/Editor/PropertyEditor/Private/UserInterface/PropertyEditor/SPropertyEditorCombo.h
@@ -27,8 +27,7 @@ public:
 
 	void GetDesiredWidth( float& OutMinDesiredWidth, float& OutMaxDesiredWidth );
 private:
-
-	void GenerateComboBoxStrings( TArray< TSharedPtr<FString> >& OutComboBoxStrings, TArray< FText >& OutToolTips, TArray<bool>& OutRestrictedItems );
+	void GenerateComboBoxStrings( TArray< TSharedPtr<FString> >& OutComboBoxStrings, TArray<TSharedPtr<SToolTip>>& OutToolTips, TArray<bool>& OutRestrictedItems );
 	void OnComboSelectionChanged( TSharedPtr<FString> NewValue, ESelectInfo::Type SelectInfo );
 	void OnComboOpening();
 
diff --git a/Engine/Source/Editor/PropertyEditor/Public/DetailWidgetRow.h b/Engine/Source/Editor/PropertyEditor/Public/DetailWidgetRow.h
--- a/Engine/Source/Editor/PropertyEditor/Public/DetailWidgetRow.h
+++ b/Engine/Source/Editor/PropertyEditor/Public/DetailWidgetRow.h
@@ -99,6 +99,7 @@ public:
 		, CopyMenuAction()
 		, PasteMenuAction()
 		, RowTagName()
+		, DiffersFromDefaultAttr( false )
 	{
 	}
 	
@@ -211,6 +212,15 @@ public:
 		return *this;
 	}
 
+	/**
+	* Sets flag to indicate if property value differs from the default
+	*/
+	FDetailWidgetRow& DiffersFromDefault(const TAttribute<bool>& InDiffersFromDefaultAttr)
+	{
+		DiffersFromDefaultAttr = InDiffersFromDefaultAttr;
+		return *this;
+	}
+
 public:
 	/** Name column content */
 	FDetailWidgetDecl NameWidget;
@@ -230,5 +240,7 @@ public:
 	FUIAction PasteMenuAction;
 	/* Tag to identify this row */
 	FName	RowTagName;
+	/* Flag to track if property has been modified from default */
+	TAttribute<bool> DiffersFromDefaultAttr;
 };
 
diff --git a/Engine/Source/Editor/PropertyEditor/Public/IDetailChildrenBuilder.h b/Engine/Source/Editor/PropertyEditor/Public/IDetailChildrenBuilder.h
--- a/Engine/Source/Editor/PropertyEditor/Public/IDetailChildrenBuilder.h
+++ b/Engine/Source/Editor/PropertyEditor/Public/IDetailChildrenBuilder.h
@@ -3,10 +3,12 @@
 #pragma once
 
 #include "CoreMinimal.h"
-#include "IDetailCustomNodeBuilder.h"
-#include "Widgets/SWidget.h"
-#include "UObject/StructOnScope.h"
-#include "PropertyHandle.h"
+
+class SWidget;
+class IPropertyHandle;
+class IDetailCategoryBuilder;
+class IDetailCustomNodeBuilder;
+class FStructOnScope;
 
 /**
  * Builder for adding children to a detail customization
@@ -64,4 +66,9 @@ public:
 	 */
 	virtual TSharedRef<SWidget> GenerateStructValueWidget(TSharedRef<IPropertyHandle> StructPropertyHandle) = 0;
 
-};
+	/**
+	 * @return the parent category on the customized object that this children is in.
+	 */
+	virtual IDetailCategoryBuilder& GetParentCategory() = 0;
+
+};
\ No newline at end of file
diff --git a/Engine/Source/Editor/PropertyEditor/Public/IDetailGroup.h b/Engine/Source/Editor/PropertyEditor/Public/IDetailGroup.h
--- a/Engine/Source/Editor/PropertyEditor/Public/IDetailGroup.h
+++ b/Engine/Source/Editor/PropertyEditor/Public/IDetailGroup.h
@@ -42,6 +42,15 @@ public:
 	 */
 	virtual class IDetailPropertyRow& AddPropertyRow( TSharedRef<IPropertyHandle> PropertyHandle ) = 0;
 
+	/**
+	 * Adds a group to the category
+	 *
+	 * @param NewGroupName	The name of the group
+	 * @param LocalizedDisplayName	The display name of the group
+	 * @param true if the group should start expanded
+	 */
+	virtual IDetailGroup& AddGroup(FName NewGroupName, const FText& InLocalizedDisplayName, bool bInStartExpanded = false) = 0;
+
 	/**
 	 * Toggles expansion on the group
 	 *
diff --git a/Engine/Source/Editor/PropertyEditor/Public/PropertyCustomizationHelpers.h b/Engine/Source/Editor/PropertyEditor/Public/PropertyCustomizationHelpers.h
--- a/Engine/Source/Editor/PropertyEditor/Public/PropertyCustomizationHelpers.h
+++ b/Engine/Source/Editor/PropertyEditor/Public/PropertyCustomizationHelpers.h
@@ -169,6 +169,8 @@ public:
 		, _AllowAbstract(false)
 		, _IsBlueprintBaseOnly(false)
 		, _AllowNone(true)
+		, _HideViewOptions(false)
+		, _ShowTreeView(false)
 	{}
 		/** The meta class that the selected class must be a child-of (required) */
 		SLATE_ARGUMENT(const UClass*, MetaClass)
@@ -180,6 +182,10 @@ public:
 		SLATE_ARGUMENT(bool, IsBlueprintBaseOnly)
 		/** Should we be able to select "None" as a class? (optional) */
 		SLATE_ARGUMENT(bool, AllowNone)
+		/** Show the View Options part of the class picker dialog*/
+		SLATE_ARGUMENT(bool, HideViewOptions)
+		/** Show the class picker as a tree view rather than a list*/
+		SLATE_ARGUMENT(bool, ShowTreeView)
 		/** Attribute used to get the currently selected class (required) */
 		SLATE_ATTRIBUTE(const UClass*, SelectedClass)
 		/** Delegate used to set the currently selected class (required) */
@@ -396,6 +402,14 @@ DECLARE_DELEGATE_TwoParams( FOnResetMaterialToDefaultClicked, UMaterialInterface
 
 DECLARE_DELEGATE_RetVal(bool, FOnMaterialListDirty);
 
+DECLARE_DELEGATE_RetVal(bool, FOnCanCopyMaterialList);
+DECLARE_DELEGATE(FOnCopyMaterialList);
+DECLARE_DELEGATE(FOnPasteMaterialList);
+
+DECLARE_DELEGATE_RetVal_OneParam(bool, FOnCanCopyMaterialItem, int32);
+DECLARE_DELEGATE_OneParam(FOnCopyMaterialItem, int32);
+DECLARE_DELEGATE_OneParam(FOnPasteMaterialItem, int32);
+
 struct FMaterialListDelegates
 {
 	FMaterialListDelegates()
@@ -418,6 +432,20 @@ struct FMaterialListDelegates
 	FOnResetMaterialToDefaultClicked OnResetMaterialToDefaultClicked;
 	/** Delegate called when we tick the material list to know if the list is dirty*/
 	FOnMaterialListDirty OnMaterialListDirty;
+
+	/** Delegate called Copying a material list */
+	FOnCopyMaterialList OnCopyMaterialList;
+	/** Delegate called to know if we can copy a material list */
+	FOnCanCopyMaterialList OnCanCopyMaterialList;
+	/** Delegate called Pasting a material list */
+	FOnPasteMaterialList OnPasteMaterialList;
+
+	/** Delegate called Copying a material item */
+	FOnCopyMaterialItem OnCopyMaterialItem;
+	/** Delegate called to know if we can copy a material item */
+	FOnCanCopyMaterialItem OnCanCopyMaterialItem;
+	/** Delegate called Pasting a material item */
+	FOnPasteMaterialItem OnPasteMaterialItem;
 };
 
 
@@ -527,6 +555,13 @@ private:
 	void AddMaterialItem(FDetailWidgetRow& Row, int32 CurrentSlot, const FMaterialListItem& Item, bool bDisplayLink);
 
 private:
+	bool OnCanCopyMaterialList() const;
+	void OnCopyMaterialList();
+	void OnPasteMaterialList();
+
+	bool OnCanCopyMaterialItem(int32 CurrentSlot) const;
+	void OnCopyMaterialItem(int32 CurrentSlot);
+	void OnPasteMaterialItem(int32 CurrentSlot);
 
 	/** Delegates for the material list */
 	FMaterialListDelegates MaterialListDelegates;
@@ -591,6 +626,14 @@ DECLARE_DELEGATE_RetVal_TwoParams(TSharedRef<SWidget>, FOnGenerateWidgetsForSect
 
 DECLARE_DELEGATE_TwoParams(FOnResetSectionToDefaultClicked, int32, int32);
 
+DECLARE_DELEGATE_RetVal(bool, FOnCanCopySectionList);
+DECLARE_DELEGATE(FOnCopySectionList);
+DECLARE_DELEGATE(FOnPasteSectionList);
+
+DECLARE_DELEGATE_RetVal_TwoParams(bool, FOnCanCopySectionItem, int32, int32);
+DECLARE_DELEGATE_TwoParams(FOnCopySectionItem, int32, int32);
+DECLARE_DELEGATE_TwoParams(FOnPasteSectionItem, int32, int32);
+
 struct FSectionListDelegates
 {
 	FSectionListDelegates()
@@ -611,6 +654,20 @@ struct FSectionListDelegates
 	FOnGenerateWidgetsForSection OnGenerateCustomSectionWidgets;
 	/** Delegate called when a Section list item should be reset to default */
 	FOnResetSectionToDefaultClicked OnResetSectionToDefaultClicked;
+
+	/** Delegate called Copying a section list */
+	FOnCopySectionList OnCopySectionList;
+	/** Delegate called to know if we can copy a section list */
+	FOnCanCopySectionList OnCanCopySectionList;
+	/** Delegate called Pasting a section list */
+	FOnPasteSectionList OnPasteSectionList;
+
+	/** Delegate called Copying a section item */
+	FOnCopySectionItem OnCopySectionItem;
+	/** Delegate called To know if we can copy a section item */
+	FOnCanCopySectionItem OnCanCopySectionItem;
+	/** Delegate called Pasting a section item */
+	FOnPasteSectionItem OnPasteSectionItem;
 };
 
 /**
@@ -737,6 +794,14 @@ private:
 	void AddSectionItem(class FDetailWidgetRow& Row, int32 LodIndex, const struct FSectionListItem& Item, bool bDisplayLink);
 
 private:
+	bool OnCanCopySectionList() const;
+	void OnCopySectionList();
+	void OnPasteSectionList();
+
+	bool OnCanCopySectionItem(int32 LODIndex, int32 SectionIndex) const;
+	void OnCopySectionItem(int32 LODIndex, int32 SectionIndex);
+	void OnPasteSectionItem(int32 LODIndex, int32 SectionIndex);
+
 	/** Delegates for the Section list */
 	FSectionListDelegates SectionListDelegates;
 	/** Called to rebuild the children of the detail tree */
diff --git a/Engine/Source/Editor/SceneOutliner/Private/ActorTreeItem.cpp b/Engine/Source/Editor/SceneOutliner/Private/ActorTreeItem.cpp
--- a/Engine/Source/Editor/SceneOutliner/Private/ActorTreeItem.cpp
+++ b/Engine/Source/Editor/SceneOutliner/Private/ActorTreeItem.cpp
@@ -246,7 +246,7 @@ FTreeItemPtr FActorTreeItem::CreateParent() const
 	{
 		return MakeShareable(new FActorTreeItem(ParentActor));
 	}
-	else
+	else if(!ParentActor)
 	{
 		const bool bShouldShowFolders = SharedData->Mode == ESceneOutlinerMode::ActorBrowsing || SharedData->bOnlyShowFolders;
 
@@ -255,11 +255,11 @@ FTreeItemPtr FActorTreeItem::CreateParent() const
 		{
 			return MakeShareable(new FFolderTreeItem(ActorFolder));
 		}
-	}
 
-	if (UWorld* World = ActorPtr->GetWorld())
-	{
-		return MakeShareable(new FWorldTreeItem(World));
+		if (UWorld* World = ActorPtr->GetWorld())
+		{
+			return MakeShareable(new FWorldTreeItem(World));
+		}
 	}
 
 	return nullptr;
diff --git a/Engine/Source/Editor/Sequencer/Private/SequencerCurveOwner.cpp b/Engine/Source/Editor/Sequencer/Private/SequencerCurveOwner.cpp
--- a/Engine/Source/Editor/Sequencer/Private/SequencerCurveOwner.cpp
+++ b/Engine/Source/Editor/Sequencer/Private/SequencerCurveOwner.cpp
@@ -162,6 +162,26 @@ void FSequencerCurveOwner::ModifyOwner()
 	}
 }
 
+TArray<const UObject*> FSequencerCurveOwner::GetOwners() const
+{
+	TArray<UMovieSceneSection*> Owners;
+	EditInfoToSectionMap.GenerateValueArray( Owners );
+
+	TArray<const UObject*> CastedOwners;
+	CastedOwners.Reserve(Owners.Num());
+	for (auto Owner : Owners)
+	{
+		const UObject* CastedOwner = Cast<const UObject>(Owner);
+
+		if (CastedOwner)
+		{
+			CastedOwners.Add(CastedOwner);
+		}
+	}
+
+	return CastedOwners;
+}
+
 void FSequencerCurveOwner::MakeTransactional()
 {
 	TArray<UMovieSceneSection*> Owners;
diff --git a/Engine/Source/Editor/SourceControlWindows/Private/SSourceControlRevert.cpp b/Engine/Source/Editor/SourceControlWindows/Private/SSourceControlRevert.cpp
--- a/Engine/Source/Editor/SourceControlWindows/Private/SSourceControlRevert.cpp
+++ b/Engine/Source/Editor/SourceControlWindows/Private/SSourceControlRevert.cpp
@@ -28,6 +28,7 @@
 #include "EditorStyleSet.h"
 #include "ISourceControlModule.h"
 #include "PackageTools.h"
+#include "Settings/EditorExperimentalSettings.h"
 
 #include "Framework/Notifications/NotificationManager.h"
 #include "Widgets/Notifications/SNotificationList.h"
@@ -438,40 +439,85 @@ bool FSourceControlWindows::PromptForRevert( const TArray<FString>& InPackageNam
 			if (PackagesToRevert.Num() > 0)
 			{
 				// attempt to unload the packages we are about to revert
-				TArray<UPackage*> PackagesToUnload;
+				TArray<UPackage*> LoadedPackages;
 				for (TArray<FString>::TConstIterator PackageIter(InPackageNames); PackageIter; ++PackageIter)
 				{
 					UPackage* Package = FindPackage(NULL, **PackageIter);
 					if (Package != NULL)
 					{
-						PackagesToUnload.Add(Package);
+						LoadedPackages.Add(Package);
 					}
 				}
 
-				if (PackagesToUnload.Num() == 0 || PackageTools::UnloadPackages(PackagesToUnload))
+				if (GetDefault<UEditorExperimentalSettings>()->bEnableContentHotReloading)
 				{
-					// Iterate over the names again, the UPackage*'s may be invalid if they have been GC'd
-					for (TArray<FString>::TConstIterator PackageIter(InPackageNames); PackageIter; ++PackageIter)
+					const TArray<FString> RevertPackageFilenames = SourceControlHelpers::PackageFilenames(PackagesToRevert);
+
+					// Prepare the packages to be reverted...
+					for (UPackage* Package : LoadedPackages)
 					{
-						UPackage* Package = FindPackage(NULL, **PackageIter);
-						if (Package != NULL)
+						// Detach the linkers of any loaded packages so that SCC can overwrite the files...
+						if (!Package->IsFullyLoaded())
 						{
-							Package->ClearFlags(RF_WasLoaded);
-							
-							// Create and display a notification about the revert failing (the object was still being used)
-							const FText NotificationErrorText = FText::Format(LOCTEXT("RevertFailed", "Failed to revert file {0} since it is currently in use."), 
-								FText::AsCultureInvariant(Package->GetName()));
-							FNotificationInfo Info(NotificationErrorText);
-							Info.ExpireDuration = 2.0f;
-							FSlateNotificationManager::Get().AddNotification(Info);
+							FlushAsyncLoading();
+							Package->FullyLoad();
 						}
+						ResetLoaders(Package);
 					}
 
-					SourceControlProvider.Execute(ISourceControlOperation::Create<FRevert>(), SourceControlHelpers::PackageFilenames(PackagesToRevert));
+					// Revert everything...
+					SourceControlProvider.Execute(ISourceControlOperation::Create<FRevert>(), RevertPackageFilenames);
+
+					// Reverting may have deleted some packages, so we need to unload those rather than re-load them...
+					TArray<UPackage*> PackagesToUnload;
+					LoadedPackages.RemoveAll([&](UPackage* InPackage) -> bool
+					{
+						const FString PackageExtension = InPackage->ContainsMap() ? FPackageName::GetMapPackageExtension() : FPackageName::GetAssetPackageExtension();
+						const FString PackageFilename = FPackageName::LongPackageNameToFilename(InPackage->GetName(), PackageExtension);
+						if (!FPaths::FileExists(PackageFilename))
+						{
+							PackagesToUnload.Emplace(InPackage);
+							return true; // remove package
+						}
+						return false; // keep package
+					});
+
+					// Hot-reload the new packages...
+					PackageTools::ReloadPackages(LoadedPackages);
+
+					// Unload any deleted packages...
+					PackageTools::UnloadPackages(PackagesToUnload);
+
+					// Re-cache the SCC state...
+					SourceControlProvider.Execute(ISourceControlOperation::Create<FUpdateStatus>(), RevertPackageFilenames, EConcurrency::Asynchronous);
 				}
 				else
 				{
-					FMessageDialog::Open(EAppMsgType::Ok, LOCTEXT("SCC_Revert_PackageUnloadFailed", "Could not unload assets when reverting files"));
+					if (LoadedPackages.Num() == 0 || PackageTools::UnloadPackages(LoadedPackages))
+					{
+						// Iterate over the names again, the UPackage*'s may be invalid if they have been GC'd
+						for (TArray<FString>::TConstIterator PackageIter(InPackageNames); PackageIter; ++PackageIter)
+						{
+							UPackage* Package = FindPackage(NULL, **PackageIter);
+							if (Package != NULL)
+							{
+								Package->ClearFlags(RF_WasLoaded);
+								
+								// Create and display a notification about the revert failing (the object was still being used)
+								const FText NotificationErrorText = FText::Format(LOCTEXT("RevertFailed", "Failed to revert file {0} since it is currently in use."), 
+									FText::AsCultureInvariant(Package->GetName()));
+								FNotificationInfo Info(NotificationErrorText);
+								Info.ExpireDuration = 2.0f;
+								FSlateNotificationManager::Get().AddNotification(Info);
+							}
+						}
+
+						SourceControlProvider.Execute(ISourceControlOperation::Create<FRevert>(), SourceControlHelpers::PackageFilenames(PackagesToRevert));
+					}
+					else
+					{
+						FMessageDialog::Open(EAppMsgType::Ok, LOCTEXT("SCC_Revert_PackageUnloadFailed", "Could not unload assets when reverting files"));
+					}
 				}
 
 				bReverted = true;
diff --git a/Engine/Source/Editor/StaticMeshEditor/Private/StaticMeshEditor.cpp b/Engine/Source/Editor/StaticMeshEditor/Private/StaticMeshEditor.cpp
--- a/Engine/Source/Editor/StaticMeshEditor/Private/StaticMeshEditor.cpp
+++ b/Engine/Source/Editor/StaticMeshEditor/Private/StaticMeshEditor.cpp
@@ -1868,6 +1868,11 @@ void FStaticMeshEditor::OnObjectReimported(UObject* InObject)
 	if(StaticMesh == InObject)
 	{
 		SetEditorMesh(Cast<UStaticMesh>(InObject));
+
+		if (SocketManager.IsValid())
+		{
+			SocketManager->UpdateStaticMesh();
+		}
 	}
 }
 
@@ -1987,6 +1992,7 @@ void FStaticMeshEditor::RedoAction()
 void FStaticMeshEditor::PostUndo( bool bSuccess )
 {
 	RemoveInvalidPrims();
+	RefreshTool();
 
 	OnPostUndo.Broadcast();
 }
@@ -1994,6 +2000,7 @@ void FStaticMeshEditor::PostUndo( bool bSuccess )
 void FStaticMeshEditor::PostRedo( bool bSuccess )
 {
 	RemoveInvalidPrims();
+	RefreshTool();
 
 	OnPostUndo.Broadcast();
 }
diff --git a/Engine/Source/Editor/StaticMeshEditor/Private/StaticMeshEditorTools.cpp b/Engine/Source/Editor/StaticMeshEditor/Private/StaticMeshEditorTools.cpp
--- a/Engine/Source/Editor/StaticMeshEditor/Private/StaticMeshEditorTools.cpp
+++ b/Engine/Source/Editor/StaticMeshEditor/Private/StaticMeshEditorTools.cpp
@@ -31,6 +31,7 @@
 #include "EngineAnalytics.h"
 #include "Widgets/Input/STextComboBox.h"
 #include "ScopedTransaction.h"
+#include "JsonObjectConverter.h"
 
 const float MaxHullAccuracy = 1.f;
 const float MinHullAccuracy = 0.f;
@@ -57,7 +58,8 @@ void FStaticMeshDetails::CustomizeDetails( class IDetailLayoutBuilder& DetailBui
 	IDetailCategoryBuilder& LODSettingsCategory = DetailBuilder.EditCategory( "LodSettings", LOCTEXT("LodSettingsCategory", "LOD Settings") );
 	IDetailCategoryBuilder& StaticMeshCategory = DetailBuilder.EditCategory( "StaticMesh", LOCTEXT("StaticMeshGeneralSettings", "Static Mesh Settings") );
 	IDetailCategoryBuilder& ImportCategory = DetailBuilder.EditCategory( "ImportSettings", LOCTEXT("ImportGeneralSettings", "Import Settings") );
-	
+	IDetailCategoryBuilder& CollisionCategory = DetailBuilder.EditCategory("Collision", LOCTEXT("CollisionCategory", "Collision"));
+
 	DetailBuilder.EditCategory( "Navigation", FText::GetEmpty(), ECategoryPriority::Uncommon );
 
 	LevelOfDetailSettings = MakeShareable( new FLevelOfDetailSettingsLayout( StaticMeshEditor ) );
@@ -97,7 +99,7 @@ void FStaticMeshDetails::CustomizeDetails( class IDetailLayoutBuilder& DetailBui
 			TSharedPtr<IPropertyHandle> ChildProp = BodyPropObject->GetChildHandle(ChildIndex);
 			if( ChildProp.IsValid() && ChildProp->GetProperty() && !HiddenBodyInstanceProps.Contains(ChildProp->GetProperty()->GetFName()) )
 			{
-				StaticMeshCategory.AddProperty( ChildProp );
+				CollisionCategory.AddProperty( ChildProp );
 			}
 		}
 	}
@@ -1370,9 +1372,218 @@ void FMeshSectionSettingsLayout::AddToCategory( IDetailCategoryBuilder& Category
 	SectionListDelegates.OnSectionChanged.BindSP(this, &FMeshSectionSettingsLayout::OnSectionChanged);
 	SectionListDelegates.OnGenerateCustomNameWidgets.BindSP(this, &FMeshSectionSettingsLayout::OnGenerateCustomNameWidgetsForSection);
 	SectionListDelegates.OnGenerateCustomSectionWidgets.BindSP(this, &FMeshSectionSettingsLayout::OnGenerateCustomSectionWidgetsForSection);
+
+	SectionListDelegates.OnCopySectionList.BindSP(this, &FMeshSectionSettingsLayout::OnCopySectionList, LODIndex);
+	SectionListDelegates.OnCanCopySectionList.BindSP(this, &FMeshSectionSettingsLayout::OnCanCopySectionList, LODIndex);
+	SectionListDelegates.OnPasteSectionList.BindSP(this, &FMeshSectionSettingsLayout::OnPasteSectionList, LODIndex);
+	SectionListDelegates.OnCopySectionItem.BindSP(this, &FMeshSectionSettingsLayout::OnCopySectionItem);
+	SectionListDelegates.OnCanCopySectionItem.BindSP(this, &FMeshSectionSettingsLayout::OnCanCopySectionItem);
+	SectionListDelegates.OnPasteSectionItem.BindSP(this, &FMeshSectionSettingsLayout::OnPasteSectionItem);
+
 	CategoryBuilder.AddCustomBuilder(MakeShareable(new FSectionList(CategoryBuilder.GetParentLayout(), SectionListDelegates, (LODIndex > 0))));
 }
 
+void FMeshSectionSettingsLayout::OnCopySectionList(int32 CurrentLODIndex)
+{
+	TSharedRef<FJsonObject> RootJsonObject = MakeShareable(new FJsonObject());
+
+	UStaticMesh& StaticMesh = GetStaticMesh();
+	FStaticMeshRenderData* RenderData = StaticMesh.RenderData.Get();
+
+	if (RenderData != nullptr && RenderData->LODResources.IsValidIndex(CurrentLODIndex))
+	{
+		FStaticMeshLODResources& LOD = RenderData->LODResources[CurrentLODIndex];
+
+		for (int32 SectionIndex = 0; SectionIndex < LOD.Sections.Num(); ++SectionIndex)
+		{
+			const FStaticMeshSection& Section = LOD.Sections[SectionIndex];
+
+			TSharedPtr<FJsonObject> JSonSection = MakeShareable(new FJsonObject);
+
+			JSonSection->SetNumberField(TEXT("MaterialIndex"), Section.MaterialIndex);
+			JSonSection->SetBoolField(TEXT("EnableCollision"), Section.bEnableCollision);
+			JSonSection->SetBoolField(TEXT("CastShadow"), Section.bCastShadow);
+
+			RootJsonObject->SetObjectField(FString::Printf(TEXT("Section_%d"), SectionIndex), JSonSection);
+		}
+	}
+
+	typedef TJsonWriter<TCHAR, TPrettyJsonPrintPolicy<TCHAR>> FStringWriter;
+	typedef TJsonWriterFactory<TCHAR, TPrettyJsonPrintPolicy<TCHAR>> FStringWriterFactory;
+
+	FString CopyStr;
+	TSharedRef<FStringWriter> Writer = FStringWriterFactory::Create(&CopyStr);
+	FJsonSerializer::Serialize(RootJsonObject, Writer);
+
+	if (!CopyStr.IsEmpty())
+	{
+		FPlatformMisc::ClipboardCopy(*CopyStr);
+	}
+}
+
+bool FMeshSectionSettingsLayout::OnCanCopySectionList(int32 CurrentLODIndex) const
+{
+	UStaticMesh& StaticMesh = GetStaticMesh();
+	FStaticMeshRenderData* RenderData = StaticMesh.RenderData.Get();
+
+	if (RenderData != nullptr && RenderData->LODResources.IsValidIndex(CurrentLODIndex))
+	{
+		FStaticMeshLODResources& LOD = RenderData->LODResources[CurrentLODIndex];
+
+		return LOD.Sections.Num() > 0;
+	}
+
+	return false;
+}
+
+void FMeshSectionSettingsLayout::OnPasteSectionList(int32 CurrentLODIndex)
+{
+	FString PastedText;
+	FPlatformMisc::ClipboardPaste(PastedText);
+
+	TSharedPtr<FJsonObject> RootJsonObject;
+	TSharedRef<TJsonReader<>> Reader = TJsonReaderFactory<>::Create(PastedText);
+	bool bResult = FJsonSerializer::Deserialize(Reader, RootJsonObject);
+
+	if (RootJsonObject.IsValid())
+	{
+		UStaticMesh& StaticMesh = GetStaticMesh();
+		FStaticMeshRenderData* RenderData = StaticMesh.RenderData.Get();
+
+		if (RenderData != nullptr && RenderData->LODResources.IsValidIndex(CurrentLODIndex))
+		{
+			UProperty* Property = UStaticMesh::StaticClass()->FindPropertyByName(GET_MEMBER_NAME_STRING_CHECKED(UStaticMesh, SectionInfoMap));
+
+			GetStaticMesh().PreEditChange(Property);
+
+			FScopedTransaction Transaction(LOCTEXT("StaticMeshToolChangedPasteSectionList", "Pasted section list"));
+			GetStaticMesh().Modify();
+
+			FStaticMeshLODResources& LOD = RenderData->LODResources[CurrentLODIndex];
+
+			for (int32 SectionIndex = 0; SectionIndex < LOD.Sections.Num(); ++SectionIndex)
+			{
+				FStaticMeshSection& Section = LOD.Sections[SectionIndex];
+
+				const TSharedPtr<FJsonObject>* JSonSection = nullptr;
+
+				if (RootJsonObject->TryGetObjectField(FString::Printf(TEXT("Section_%d"), SectionIndex), JSonSection))
+				{
+					(*JSonSection)->TryGetNumberField(TEXT("MaterialIndex"), Section.MaterialIndex);
+					(*JSonSection)->TryGetBoolField(TEXT("EnableCollision"), Section.bEnableCollision);
+					(*JSonSection)->TryGetBoolField(TEXT("CastShadow"), Section.bCastShadow);
+
+					// Update the section info
+					FMeshSectionInfo Info = StaticMesh.SectionInfoMap.Get(LODIndex, SectionIndex);
+					Info.MaterialIndex = Section.MaterialIndex;
+					Info.bCastShadow = Section.bCastShadow;
+					Info.bEnableCollision = Section.bEnableCollision;
+
+					StaticMesh.SectionInfoMap.Set(LODIndex, SectionIndex, Info);
+				}
+			}
+
+			CallPostEditChange(Property);
+		}
+	}
+}
+
+void FMeshSectionSettingsLayout::OnCopySectionItem(int32 CurrentLODIndex, int32 SectionIndex)
+{
+	TSharedRef<FJsonObject> RootJsonObject = MakeShareable(new FJsonObject());
+
+	UStaticMesh& StaticMesh = GetStaticMesh();
+	FStaticMeshRenderData* RenderData = StaticMesh.RenderData.Get();
+
+	if (RenderData != nullptr && RenderData->LODResources.IsValidIndex(CurrentLODIndex))
+	{
+		FStaticMeshLODResources& LOD = RenderData->LODResources[CurrentLODIndex];
+
+		if (LOD.Sections.IsValidIndex(SectionIndex))
+		{
+			const FStaticMeshSection& Section = LOD.Sections[SectionIndex];
+
+			RootJsonObject->SetNumberField(TEXT("MaterialIndex"), Section.MaterialIndex);
+			RootJsonObject->SetBoolField(TEXT("EnableCollision"), Section.bEnableCollision);
+			RootJsonObject->SetBoolField(TEXT("CastShadow"), Section.bCastShadow);
+		}
+	}
+
+	typedef TJsonWriter<TCHAR, TPrettyJsonPrintPolicy<TCHAR>> FStringWriter;
+	typedef TJsonWriterFactory<TCHAR, TPrettyJsonPrintPolicy<TCHAR>> FStringWriterFactory;
+
+	FString CopyStr;
+	TSharedRef<FStringWriter> Writer = FStringWriterFactory::Create(&CopyStr);
+	FJsonSerializer::Serialize(RootJsonObject, Writer);
+
+	if (!CopyStr.IsEmpty())
+	{
+		FPlatformMisc::ClipboardCopy(*CopyStr);
+	}
+}
+
+bool FMeshSectionSettingsLayout::OnCanCopySectionItem(int32 CurrentLODIndex, int32 SectionIndex) const
+{
+	UStaticMesh& StaticMesh = GetStaticMesh();
+	FStaticMeshRenderData* RenderData = StaticMesh.RenderData.Get();
+
+	if (RenderData != nullptr && RenderData->LODResources.IsValidIndex(CurrentLODIndex))
+	{
+		FStaticMeshLODResources& LOD = RenderData->LODResources[CurrentLODIndex];
+
+		return LOD.Sections.IsValidIndex(SectionIndex);
+	}
+
+	return false;
+}
+
+void FMeshSectionSettingsLayout::OnPasteSectionItem(int32 CurrentLODIndex, int32 SectionIndex)
+{
+	FString PastedText;
+	FPlatformMisc::ClipboardPaste(PastedText);
+
+	TSharedPtr<FJsonObject> RootJsonObject;
+	TSharedRef<TJsonReader<>> Reader = TJsonReaderFactory<>::Create(PastedText);
+	bool bResult = FJsonSerializer::Deserialize(Reader, RootJsonObject);
+
+	if (RootJsonObject.IsValid())
+	{
+		UStaticMesh& StaticMesh = GetStaticMesh();
+		FStaticMeshRenderData* RenderData = StaticMesh.RenderData.Get();
+
+		if (RenderData != nullptr && RenderData->LODResources.IsValidIndex(CurrentLODIndex))
+		{
+			UProperty* Property = UStaticMesh::StaticClass()->FindPropertyByName(GET_MEMBER_NAME_STRING_CHECKED(UStaticMesh, SectionInfoMap));
+
+			GetStaticMesh().PreEditChange(Property);
+
+			FScopedTransaction Transaction(LOCTEXT("StaticMeshToolChangedPasteSectionItem", "Pasted section item"));
+			GetStaticMesh().Modify();
+
+			FStaticMeshLODResources& LOD = RenderData->LODResources[CurrentLODIndex];
+
+			if (LOD.Sections.IsValidIndex(SectionIndex))
+			{
+				FStaticMeshSection& Section = LOD.Sections[SectionIndex];
+
+				RootJsonObject->TryGetNumberField(TEXT("MaterialIndex"), Section.MaterialIndex);
+				RootJsonObject->TryGetBoolField(TEXT("EnableCollision"), Section.bEnableCollision);
+				RootJsonObject->TryGetBoolField(TEXT("CastShadow"), Section.bCastShadow);
+
+				// Update the section info
+				FMeshSectionInfo Info = StaticMesh.SectionInfoMap.Get(LODIndex, SectionIndex);
+				Info.MaterialIndex = Section.MaterialIndex;
+				Info.bCastShadow = Section.bCastShadow;
+				Info.bEnableCollision = Section.bEnableCollision;
+
+				StaticMesh.SectionInfoMap.Set(LODIndex, SectionIndex, Info);
+			}
+
+			CallPostEditChange(Property);
+		}
+	}
+}
+
 void FMeshSectionSettingsLayout::OnGetSectionsForView(ISectionListBuilder& OutSections, int32 ForLODIndex)
 {
 	check(LODIndex == ForLODIndex);
@@ -1671,6 +1882,8 @@ UStaticMesh& FMeshMaterialsLayout::GetStaticMesh() const
 void FMeshMaterialsLayout::AddToCategory(IDetailCategoryBuilder& CategoryBuilder)
 {
 	CategoryBuilder.AddCustomRow(LOCTEXT("AddLODLevelCategories_MaterialArrayOperationAdd", "Add Material Slot"))
+		.CopyAction(FUIAction(FExecuteAction::CreateSP(this, &FMeshMaterialsLayout::OnCopyMaterialList), FCanExecuteAction::CreateSP(this, &FMeshMaterialsLayout::OnCanCopyMaterialList)))
+		.PasteAction(FUIAction(FExecuteAction::CreateSP(this, &FMeshMaterialsLayout::OnPasteMaterialList)))
 		.NameContent()
 		.HAlign(HAlign_Left)
 		.VAlign(VAlign_Center)
@@ -1719,6 +1932,7 @@ void FMeshMaterialsLayout::AddToCategory(IDetailCategoryBuilder& CategoryBuilder
 				]
 			]
 		];
+
 	FMaterialListDelegates MaterialListDelegates;
 	MaterialListDelegates.OnGetMaterials.BindSP(this, &FMeshMaterialsLayout::GetMaterials);
 	MaterialListDelegates.OnMaterialChanged.BindSP(this, &FMeshMaterialsLayout::OnMaterialChanged);
@@ -1726,9 +1940,119 @@ void FMeshMaterialsLayout::AddToCategory(IDetailCategoryBuilder& CategoryBuilder
 	MaterialListDelegates.OnMaterialListDirty.BindSP(this, &FMeshMaterialsLayout::OnMaterialListDirty);
 	MaterialListDelegates.OnResetMaterialToDefaultClicked.BindSP(this, &FMeshMaterialsLayout::OnResetMaterialToDefaultClicked);
 
+	MaterialListDelegates.OnCopyMaterialItem.BindSP(this, &FMeshMaterialsLayout::OnCopyMaterialItem);
+	MaterialListDelegates.OnCanCopyMaterialItem.BindSP(this, &FMeshMaterialsLayout::OnCanCopyMaterialItem);
+	MaterialListDelegates.OnPasteMaterialItem.BindSP(this, &FMeshMaterialsLayout::OnPasteMaterialItem);
+
 	CategoryBuilder.AddCustomBuilder(MakeShareable(new FMaterialList(CategoryBuilder.GetParentLayout(), MaterialListDelegates)));
 }
 
+void FMeshMaterialsLayout::OnCopyMaterialList()
+{
+	UProperty* Property = UStaticMesh::StaticClass()->FindPropertyByName(GET_MEMBER_NAME_STRING_CHECKED(UStaticMesh, StaticMaterials));
+	check(Property != nullptr);
+
+	auto JsonValue = FJsonObjectConverter::UPropertyToJsonValue(Property, &GetStaticMesh().StaticMaterials, 0, 0);
+
+	typedef TJsonWriter<TCHAR, TPrettyJsonPrintPolicy<TCHAR>> FStringWriter;
+	typedef TJsonWriterFactory<TCHAR, TPrettyJsonPrintPolicy<TCHAR>> FStringWriterFactory;
+
+	FString CopyStr;
+	TSharedRef<FStringWriter> Writer = FStringWriterFactory::Create(&CopyStr);
+	FJsonSerializer::Serialize(JsonValue.ToSharedRef(), TEXT(""), Writer);
+
+	if (!CopyStr.IsEmpty())
+	{
+		FPlatformMisc::ClipboardCopy(*CopyStr);
+	}
+}
+
+bool FMeshMaterialsLayout::OnCanCopyMaterialList() const
+{
+	return GetStaticMesh().StaticMaterials.Num() > 0;
+}
+
+void FMeshMaterialsLayout::OnPasteMaterialList()
+{
+	FString PastedText;
+	FPlatformMisc::ClipboardPaste(PastedText);
+
+	TSharedPtr<FJsonValue> RootJsonValue;
+	TSharedRef<TJsonReader<>> Reader = TJsonReaderFactory<>::Create(PastedText);
+	FJsonSerializer::Deserialize(Reader, RootJsonValue);
+
+	if (RootJsonValue.IsValid())
+	{
+		UProperty* Property = UStaticMesh::StaticClass()->FindPropertyByName(GET_MEMBER_NAME_STRING_CHECKED(UStaticMesh, StaticMaterials));
+		check(Property != nullptr);
+
+		GetStaticMesh().PreEditChange(Property);
+		FScopedTransaction Transaction(LOCTEXT("StaticMeshToolChangedPasteMaterialList", "Pasted material list"));
+		GetStaticMesh().Modify();
+
+		FJsonObjectConverter::JsonValueToUProperty(RootJsonValue, Property, &GetStaticMesh().StaticMaterials, 0, 0);
+
+		CallPostEditChange(Property);
+	}
+}
+
+void FMeshMaterialsLayout::OnCopyMaterialItem(int32 CurrentSlot)
+{
+	TSharedRef<FJsonObject> RootJsonObject = MakeShareable(new FJsonObject());
+
+	if (GetStaticMesh().StaticMaterials.IsValidIndex(CurrentSlot))
+	{
+		const FStaticMaterial &Material = GetStaticMesh().StaticMaterials[CurrentSlot];
+
+		FJsonObjectConverter::UStructToJsonObject(FStaticMaterial::StaticStruct(), &Material, RootJsonObject, 0, 0);
+	}
+
+	typedef TJsonWriter<TCHAR, TPrettyJsonPrintPolicy<TCHAR>> FStringWriter;
+	typedef TJsonWriterFactory<TCHAR, TPrettyJsonPrintPolicy<TCHAR>> FStringWriterFactory;
+
+	FString CopyStr;
+	TSharedRef<FStringWriter> Writer = FStringWriterFactory::Create(&CopyStr);
+	FJsonSerializer::Serialize(RootJsonObject, Writer);
+
+	if (!CopyStr.IsEmpty())
+	{
+		FPlatformMisc::ClipboardCopy(*CopyStr);
+	}
+}
+
+bool FMeshMaterialsLayout::OnCanCopyMaterialItem(int32 CurrentSlot) const
+{
+	return GetStaticMesh().StaticMaterials.IsValidIndex(CurrentSlot);
+}
+
+void FMeshMaterialsLayout::OnPasteMaterialItem(int32 CurrentSlot)
+{
+	FString PastedText;
+	FPlatformMisc::ClipboardPaste(PastedText);
+
+	TSharedPtr<FJsonObject> RootJsonObject;
+	TSharedRef<TJsonReader<>> Reader = TJsonReaderFactory<>::Create(PastedText);
+	FJsonSerializer::Deserialize(Reader, RootJsonObject);
+
+	if (RootJsonObject.IsValid())
+	{
+		UProperty* Property = UStaticMesh::StaticClass()->FindPropertyByName(GET_MEMBER_NAME_STRING_CHECKED(UStaticMesh, StaticMaterials));
+		check(Property != nullptr);
+
+		GetStaticMesh().PreEditChange(Property);
+
+		FScopedTransaction Transaction(LOCTEXT("StaticMeshToolChangedPasteMaterialItem", "Pasted material item"));
+		GetStaticMesh().Modify();
+
+		if (GetStaticMesh().StaticMaterials.IsValidIndex(CurrentSlot))
+		{
+			FJsonObjectConverter::JsonObjectToUStruct(RootJsonObject.ToSharedRef(), FSkeletalMaterial::StaticStruct(), &GetStaticMesh().StaticMaterials[CurrentSlot], 0, 0);
+		}
+
+		CallPostEditChange(Property);
+	}
+}
+
 FReply FMeshMaterialsLayout::AddMaterialSlot()
 {
 	UStaticMesh& StaticMesh = GetStaticMesh();
@@ -2293,6 +2617,50 @@ void FLevelOfDetailSettingsLayout::AddToDetailsPanel( IDetailLayoutBuilder& Deta
 	AddLODLevelCategories( DetailBuilder );
 }
 
+bool FLevelOfDetailSettingsLayout::CanRemoveLOD(int32 LODIndex) const
+{
+	UStaticMesh* StaticMesh = StaticMeshEditor.GetStaticMesh();
+
+	if (StaticMesh != nullptr)
+	{
+		const int32 NumLODs = StaticMesh->GetNumLODs();
+		
+		// LOD0 should never be removed
+		return (NumLODs > 1 && LODIndex > 0 && LODIndex < NumLODs);
+	}
+
+	return false;
+}
+
+FReply FLevelOfDetailSettingsLayout::OnRemoveLOD(int32 LODIndex)
+{
+	UStaticMesh* StaticMesh = StaticMeshEditor.GetStaticMesh();
+
+	if (StaticMesh != nullptr)
+	{
+		const int32 NumLODs = StaticMesh->GetNumLODs();
+
+		if (NumLODs > 1 && LODIndex > 0 && LODIndex < NumLODs)
+		{
+			FText RemoveLODText = FText::Format( LOCTEXT("ConfirmRemoveLOD", "Are you sure you want to remove LOD {0} from {1}?"), LODIndex, FText::FromString(StaticMesh->GetName()) );
+
+			if (FMessageDialog::Open(EAppMsgType::YesNo, RemoveLODText) == EAppReturnType::Yes)
+			{
+				FText TransactionDescription = FText::Format( LOCTEXT("OnRemoveLOD", "Remove LOD {0}"), LODIndex);
+				FScopedTransaction Transaction( TEXT(""), TransactionDescription, StaticMesh );
+
+				StaticMesh->Modify();
+				StaticMesh->SourceModels.RemoveAt(LODIndex);
+				--LODCount;
+				StaticMesh->PostEditChange();
+
+				StaticMeshEditor.RefreshTool();
+			}
+		}
+	}
+
+	return FReply::Handled();
+}
 
 void FLevelOfDetailSettingsLayout::AddLODLevelCategories( IDetailLayoutBuilder& DetailBuilder )
 {
@@ -2360,6 +2728,24 @@ void FLevelOfDetailSettingsLayout::AddLODLevelCategories( IDetailLayoutBuilder&
 
 			IDetailCategoryBuilder& LODCategory = DetailBuilder.EditCategory( *CategoryName, LODLevelString, ECategoryPriority::Important );
 
+			if (LODIndex != 0)
+			{
+				LODCategory.AddCustomRow( LOCTEXT("RemoveLOD", "Remove LOD") )
+				.ValueContent()
+				.HAlign(HAlign_Left)
+				[
+					SNew(SButton)
+					.OnClicked(this, &FLevelOfDetailSettingsLayout::OnRemoveLOD, LODIndex)
+					.IsEnabled(this, &FLevelOfDetailSettingsLayout::CanRemoveLOD, LODIndex)
+					.ToolTipText( LOCTEXT("RemoveLOD_ToolTip", "Removes this LOD from the Static Mesh") )
+					[
+						SNew(STextBlock)
+						.Text( LOCTEXT("RemoveLOD", "Remove LOD") )
+						.Font( DetailBuilder.GetDetailFont() )
+					]
+				];
+			}
+
 			LODCategory.HeaderContent
 			(
 				SNew( SBox )
diff --git a/Engine/Source/Editor/StaticMeshEditor/Private/StaticMeshEditorTools.h b/Engine/Source/Editor/StaticMeshEditor/Private/StaticMeshEditorTools.h
--- a/Engine/Source/Editor/StaticMeshEditor/Private/StaticMeshEditorTools.h
+++ b/Engine/Source/Editor/StaticMeshEditor/Private/StaticMeshEditorTools.h
@@ -268,6 +268,14 @@ private:
 	
 	UStaticMesh& GetStaticMesh() const;
 
+	void OnCopySectionList(int32 CurrentLODIndex);
+	bool OnCanCopySectionList(int32 CurrentLODIndex) const;
+	void OnPasteSectionList(int32 CurrentLODIndex);
+
+	void OnCopySectionItem(int32 CurrentLODIndex, int32 SectionIndex);
+	bool OnCanCopySectionItem(int32 CurrentLODIndex, int32 SectionIndex) const;
+	void OnPasteSectionItem(int32 CurrentLODIndex, int32 SectionIndex);
+
 	/**
 	* Called by the material list widget when we need to get new materials for the list
 	*
@@ -389,6 +397,14 @@ private:
 
 	void CallPostEditChange(UProperty* PropertyChanged = nullptr);
 
+	void OnCopyMaterialList();
+	bool OnCanCopyMaterialList() const;
+	void OnPasteMaterialList();
+
+	void OnCopyMaterialItem(int32 CurrentSlot);
+	bool OnCanCopyMaterialItem(int32 CurrentSlot) const;
+	void OnPasteMaterialItem(int32 CurrentSlot);
+
 	IStaticMeshEditor& StaticMeshEditor;
 	
 	/* This is to know if material are used by any LODs sections. */
@@ -427,6 +443,9 @@ private:
 	void OnMinLODCommitted(int32 InValue, ETextCommit::Type CommitInfo);
 	int32 GetMinLOD() const;
 
+	bool CanRemoveLOD(int32 LODIndex) const;
+	FReply OnRemoveLOD(int32 LODIndex);
+
 	float GetLODScreenSize(int32 LODIndex)const;
 	FText GetLODScreenSizeTitle(int32 LODIndex) const;
 	bool CanChangeLODScreenSize() const;
diff --git a/Engine/Source/Editor/StaticMeshEditor/Private/StaticMeshEditorViewportClient.cpp b/Engine/Source/Editor/StaticMeshEditor/Private/StaticMeshEditorViewportClient.cpp
--- a/Engine/Source/Editor/StaticMeshEditor/Private/StaticMeshEditorViewportClient.cpp
+++ b/Engine/Source/Editor/StaticMeshEditor/Private/StaticMeshEditorViewportClient.cpp
@@ -439,6 +439,14 @@ void FStaticMeshEditorViewportClient::Draw(const FSceneView* View,FPrimitiveDraw
 {
 	FEditorViewportClient::Draw(View, PDI);
 
+	TSharedPtr<IStaticMeshEditor> StaticMeshEditor = StaticMeshEditorPtr.Pin();
+
+	if(!StaticMesh->RenderData->LODResources.IsValidIndex(StaticMeshEditor->GetCurrentLODIndex()))
+	{
+		// Guard against corrupted meshes
+		return;
+	}
+
 	if (bShowCollision && StaticMesh->BodySetup)
 	{
 		// Ensure physics mesh is created before we try and draw it
@@ -457,7 +465,7 @@ void FStaticMeshEditorViewportClient::Draw(const FSceneView* View,FPrimitiveDraw
 			HSMECollisionProxy* HitProxy = new HSMECollisionProxy(KPT_Sphere, i);
 			PDI->SetHitProxy(HitProxy);
 
-			const FColor CollisionColor = StaticMeshEditorPtr.Pin()->IsSelectedPrim(HitProxy->PrimData) ? SelectedColor : UnselectedColor;
+			const FColor CollisionColor = StaticMeshEditor->IsSelectedPrim(HitProxy->PrimData) ? SelectedColor : UnselectedColor;
 			const FKSphereElem& SphereElem = AggGeom->SphereElems[i];
 			const FTransform ElemTM = SphereElem.GetTransform();
 			SphereElem.DrawElemWire(PDI, ElemTM, VectorScaleOne, CollisionColor);
@@ -470,7 +478,7 @@ void FStaticMeshEditorViewportClient::Draw(const FSceneView* View,FPrimitiveDraw
 			HSMECollisionProxy* HitProxy = new HSMECollisionProxy(KPT_Box, i);
 			PDI->SetHitProxy(HitProxy);
 
-			const FColor CollisionColor = StaticMeshEditorPtr.Pin()->IsSelectedPrim(HitProxy->PrimData) ? SelectedColor : UnselectedColor;
+			const FColor CollisionColor = StaticMeshEditor->IsSelectedPrim(HitProxy->PrimData) ? SelectedColor : UnselectedColor;
 			const FKBoxElem& BoxElem = AggGeom->BoxElems[i];
 			const FTransform ElemTM = BoxElem.GetTransform();
 			BoxElem.DrawElemWire(PDI, ElemTM, VectorScaleOne, CollisionColor);
@@ -483,7 +491,7 @@ void FStaticMeshEditorViewportClient::Draw(const FSceneView* View,FPrimitiveDraw
 			HSMECollisionProxy* HitProxy = new HSMECollisionProxy(KPT_Sphyl, i);
 			PDI->SetHitProxy(HitProxy);
 
-			const FColor CollisionColor = StaticMeshEditorPtr.Pin()->IsSelectedPrim(HitProxy->PrimData) ? SelectedColor : UnselectedColor;
+			const FColor CollisionColor = StaticMeshEditor->IsSelectedPrim(HitProxy->PrimData) ? SelectedColor : UnselectedColor;
 			const FKSphylElem& SphylElem = AggGeom->SphylElems[i];
 			const FTransform ElemTM = SphylElem.GetTransform();
 			SphylElem.DrawElemWire(PDI, ElemTM, VectorScaleOne, CollisionColor);
@@ -496,7 +504,7 @@ void FStaticMeshEditorViewportClient::Draw(const FSceneView* View,FPrimitiveDraw
 			HSMECollisionProxy* HitProxy = new HSMECollisionProxy(KPT_Convex, i);
 			PDI->SetHitProxy(HitProxy);
 
-			const FColor CollisionColor = StaticMeshEditorPtr.Pin()->IsSelectedPrim(HitProxy->PrimData) ? SelectedColor : UnselectedColor;
+			const FColor CollisionColor = StaticMeshEditor->IsSelectedPrim(HitProxy->PrimData) ? SelectedColor : UnselectedColor;
 			const FKConvexElem& ConvexElem = AggGeom->ConvexElems[i];
 			const FTransform ElemTM = ConvexElem.GetTransform();
 			ConvexElem.DrawElemWire(PDI, ElemTM, 1.f, CollisionColor);
@@ -543,7 +551,7 @@ void FStaticMeshEditorViewportClient::Draw(const FSceneView* View,FPrimitiveDraw
 
 	if( bDrawNormals || bDrawTangents || bDrawBinormals || bDrawVertices )
 	{
-		FStaticMeshLODResources& LODModel = StaticMesh->RenderData->LODResources[StaticMeshEditorPtr.Pin()->GetCurrentLODIndex()];
+		FStaticMeshLODResources& LODModel = StaticMesh->RenderData->LODResources[StaticMeshEditor->GetCurrentLODIndex()];
 		FIndexArrayView Indices = LODModel.IndexBuffer.GetArrayView();
 		uint32 NumIndices = Indices.Num();
 
@@ -815,7 +823,7 @@ void FStaticMeshEditorViewportClient::DrawCanvas( FViewport& InViewport, FSceneV
 
 	StaticMeshEditorViewport->PopulateOverlayText(TextItems);
 
-	if(bDrawUVs)
+	if(bDrawUVs && StaticMesh->RenderData->LODResources.Num() > 0)
 	{
 		const int32 YPos = 160;
 		DrawUVsForMesh(Viewport, &Canvas, YPos);
diff --git a/Engine/Source/Editor/StaticMeshEditor/StaticMeshEditor.Build.cs b/Engine/Source/Editor/StaticMeshEditor/StaticMeshEditor.Build.cs
--- a/Engine/Source/Editor/StaticMeshEditor/StaticMeshEditor.Build.cs
+++ b/Engine/Source/Editor/StaticMeshEditor/StaticMeshEditor.Build.cs
@@ -34,7 +34,9 @@ public class StaticMeshEditor : ModuleRules
 				"RawMesh",
                 "PropertyEditor",
 				"MeshUtilities",
-			}
+                "Json",
+                "JsonUtilities",
+            }
 		);
 
 		DynamicallyLoadedModuleNames.AddRange(
diff --git a/Engine/Source/Editor/TextureEditor/Private/Models/TextureEditorViewportClient.cpp b/Engine/Source/Editor/TextureEditor/Private/Models/TextureEditorViewportClient.cpp
--- a/Engine/Source/Editor/TextureEditor/Private/Models/TextureEditorViewportClient.cpp
+++ b/Engine/Source/Editor/TextureEditor/Private/Models/TextureEditorViewportClient.cpp
@@ -129,17 +129,20 @@ void FTextureEditorViewportClient::Draw(FViewport* Viewport, FCanvas* Canvas)
 
 		float Exposure = FMath::Pow(2.0f, (float)TextureEditorViewportPtr.Pin()->GetExposureBias());
 
-		FCanvasTileItem TileItem( FVector2D( XPos, YPos ), Texture->Resource, FVector2D( Width, Height ), FLinearColor(Exposure, Exposure, Exposure) );
-		TileItem.BlendMode = TextureEditorPtr.Pin()->GetColourChannelBlendMode();
-		TileItem.BatchedElementParameters = BatchedElementParameters;
-		Canvas->DrawItem( TileItem );
-
-		// Draw a white border around the texture to show its extents
-		if (Settings.TextureBorderEnabled)
+		if ( Texture->Resource != nullptr )
 		{
-			FCanvasBoxItem BoxItem( FVector2D(XPos, YPos), FVector2D(Width , Height ) );
-			BoxItem.SetColor( Settings.TextureBorderColor );
-			Canvas->DrawItem( BoxItem );
+			FCanvasTileItem TileItem( FVector2D( XPos, YPos ), Texture->Resource, FVector2D( Width, Height ), FLinearColor(Exposure, Exposure, Exposure) );
+			TileItem.BlendMode = TextureEditorPtr.Pin()->GetColourChannelBlendMode();
+			TileItem.BatchedElementParameters = BatchedElementParameters;
+			Canvas->DrawItem( TileItem );
+
+			// Draw a white border around the texture to show its extents
+			if (Settings.TextureBorderEnabled)
+			{
+				FCanvasBoxItem BoxItem( FVector2D(XPos, YPos), FVector2D(Width , Height ) );
+				BoxItem.SetColor( Settings.TextureBorderColor );
+				Canvas->DrawItem( BoxItem );
+			}
 		}
 	}
 }
@@ -323,4 +326,4 @@ void FTextureEditorViewportClient::DestroyCheckerboardTexture()
 		CheckerboardTexture->MarkPendingKill();
 		CheckerboardTexture = NULL;
 	}
-}
+}
\ No newline at end of file
diff --git a/Engine/Source/Editor/TextureEditor/Private/TextureEditorToolkit.cpp b/Engine/Source/Editor/TextureEditor/Private/TextureEditorToolkit.cpp
--- a/Engine/Source/Editor/TextureEditor/Private/TextureEditorToolkit.cpp
+++ b/Engine/Source/Editor/TextureEditor/Private/TextureEditorToolkit.cpp
@@ -333,6 +333,12 @@ UTexture* FTextureEditorToolkit::GetTexture( ) const
 }
 
 
+bool FTextureEditorToolkit::HasValidTextureResource( ) const
+{
+	return Texture != nullptr && Texture->Resource != nullptr;
+}
+
+
 bool FTextureEditorToolkit::GetUseSpecifiedMip( ) const
 {
 	if (GetMaxMipLevel().Get(MIPLEVEL_MAX) > 0)
diff --git a/Engine/Source/Editor/TextureEditor/Private/TextureEditorToolkit.h b/Engine/Source/Editor/TextureEditor/Private/TextureEditorToolkit.h
--- a/Engine/Source/Editor/TextureEditor/Private/TextureEditorToolkit.h
+++ b/Engine/Source/Editor/TextureEditor/Private/TextureEditorToolkit.h
@@ -62,6 +62,7 @@ public:
 	virtual bool GetFitToViewport( ) const override;
 	virtual int32 GetMipLevel( ) const override;
 	virtual UTexture* GetTexture( ) const override;
+	virtual bool HasValidTextureResource( ) const override;
 	virtual bool GetUseSpecifiedMip( ) const override;
 	virtual double GetZoom( ) const override;
 	virtual void PopulateQuickInfo( ) override;
diff --git a/Engine/Source/Editor/TextureEditor/Private/Widgets/STextureEditorViewport.cpp b/Engine/Source/Editor/TextureEditor/Private/Widgets/STextureEditorViewport.cpp
--- a/Engine/Source/Editor/TextureEditor/Private/Widgets/STextureEditorViewport.cpp
+++ b/Engine/Source/Editor/TextureEditor/Private/Widgets/STextureEditorViewport.cpp
@@ -72,7 +72,9 @@ void STextureEditorViewport::Construct( const FArguments& InArgs, const TSharedR
 	
 	if (InToolkit->GetTexture() != nullptr)
 	{
-		TextureName = FText::FromName(InToolkit->GetTexture()->GetFName());
+		FText FormattedText = InToolkit->HasValidTextureResource() ? FText::FromString(TEXT("{0}")) : LOCTEXT( "InvalidTexture", "{0} (Invalid Texture)");
+
+		TextureName = FText::Format(FormattedText, FText::FromName(InToolkit->GetTexture()->GetFName()));
 	}
 
 	this->ChildSlot
@@ -397,6 +399,11 @@ bool STextureEditorViewport::IsZoomMenuFitChecked() const
 	return ToolkitPtr.Pin()->GetFitToViewport();
 }
 
+bool STextureEditorViewport::HasValidTextureResource() const
+{
+	return ToolkitPtr.Pin()->HasValidTextureResource();
+}
+
 FText STextureEditorViewport::HandleZoomPercentageText( ) const
 {
 	const bool bFitToViewport = ToolkitPtr.Pin()->GetFitToViewport();
diff --git a/Engine/Source/Editor/TextureEditor/Private/Widgets/STextureEditorViewport.h b/Engine/Source/Editor/TextureEditor/Private/Widgets/STextureEditorViewport.h
--- a/Engine/Source/Editor/TextureEditor/Private/Widgets/STextureEditorViewport.h
+++ b/Engine/Source/Editor/TextureEditor/Private/Widgets/STextureEditorViewport.h
@@ -119,6 +119,9 @@ private:
 	// Callback for getting the zoom slider's value.
 	float HandleZoomSliderValue( ) const;
 
+	// Checks if the texture being edited has a valid texture resource
+	bool HasValidTextureResource( ) const;
+
 private:
 
 	// Which exposure level should be used, in FStop e.g. 0:original, -1:half as bright, 1:2x as bright, 2:4x as bright.
diff --git a/Engine/Source/Editor/TextureEditor/Public/Interfaces/ITextureEditorToolkit.h b/Engine/Source/Editor/TextureEditor/Public/Interfaces/ITextureEditorToolkit.h
--- a/Engine/Source/Editor/TextureEditor/Public/Interfaces/ITextureEditorToolkit.h
+++ b/Engine/Source/Editor/TextureEditor/Public/Interfaces/ITextureEditorToolkit.h
@@ -19,6 +19,9 @@ public:
 	/** Returns the Texture asset being inspected by the Texture editor */
 	virtual UTexture* GetTexture() const = 0;
 
+	/** Returns if the Texture asset being inspected has a valid texture resource */
+	virtual bool HasValidTextureResource() const = 0;
+
 	/** Refreshes the quick info panel */
 	virtual void PopulateQuickInfo() = 0;
 
diff --git a/Engine/Source/Editor/UMGEditor/Classes/WidgetGraphSchema.h b/Engine/Source/Editor/UMGEditor/Classes/WidgetGraphSchema.h
--- a/Engine/Source/Editor/UMGEditor/Classes/WidgetGraphSchema.h
+++ b/Engine/Source/Editor/UMGEditor/Classes/WidgetGraphSchema.h
@@ -12,25 +12,6 @@ class UWidgetGraphSchema : public UEdGraphSchema_K2
 {
 	GENERATED_UCLASS_BODY()
 
-	virtual bool ShouldAlwaysPurgeOnModification() const override { return true; }
-
-	UPROPERTY()
-	FName NAME_NeverAsPin;
-
-	UPROPERTY()
-	FName NAME_PinHiddenByDefault;
-
-	UPROPERTY()
-	FName NAME_PinShownByDefault;
-
-	UPROPERTY()
-	FName NAME_AlwaysAsPin;
-
-	UPROPERTY()
-	FName NAME_OnEvaluate;
-	
-	UPROPERTY()
-	FName DefaultEvaluationHandlerName;
-	
-	// TODO UMG - Add more things specific to the needs of the widget blueprint graph.
+public:
+	virtual void BackwardCompatibilityNodeConversion(UEdGraph* Graph, bool bOnlySafeChanges) const override;
 };
diff --git a/Engine/Source/Editor/UMGEditor/Private/Designer/DesignTimeUtils.cpp b/Engine/Source/Editor/UMGEditor/Private/Designer/DesignTimeUtils.cpp
--- a/Engine/Source/Editor/UMGEditor/Private/Designer/DesignTimeUtils.cpp
+++ b/Engine/Source/Editor/UMGEditor/Private/Designer/DesignTimeUtils.cpp
@@ -42,31 +42,11 @@ bool FDesignTimeUtils::GetArrangedWidgetRelativeToWindow(TSharedRef<SWidget> Wid
 	{
 		ArrangedWidget = WidgetPath.FindArrangedWidget(Widget).Get(FArrangedWidget::NullWidget);
 		ArrangedWidget.Geometry.AppendTransform(FSlateLayoutTransform(Inverse(CurrentWindowRef->GetPositionInScreen())));
+		//ArrangedWidget.Geometry.AppendTransform(Inverse(CurrentWindowRef->GetLocalToScreenTransform()));
 		return true;
 	}
 
 	return false;
 }
 
-bool FDesignTimeUtils::GetArrangedWidgetRelativeToParent(TSharedRef<const SWidget> Widget, TSharedRef<const SWidget> Parent, FArrangedWidget& ArrangedWidget)
-{
-	FWidgetPath WidgetPath;
-	if ( FSlateApplication::Get().GeneratePathToWidgetUnchecked(Widget, WidgetPath) )
-	{
-		GetArrangedWidgetRelativeToParent(WidgetPath, Widget, Parent, ArrangedWidget);
-		return true;
-	}
-
-	return false;
-}
-
-void FDesignTimeUtils::GetArrangedWidgetRelativeToParent(FWidgetPath& WidgetPath, TSharedRef<const SWidget> Widget, TSharedRef<const SWidget> Parent, FArrangedWidget& ArrangedWidget)
-{
-	FArrangedWidget ArrangedDesigner = WidgetPath.FindArrangedWidget(Parent).Get(FArrangedWidget::NullWidget);
-
-	ArrangedWidget = WidgetPath.FindArrangedWidget(Widget).Get(FArrangedWidget::NullWidget);
-	// !!! WRH 2014/08/26 - This code tries to mutate an FGeometry in an unsupported way. This code should be refactored.
-	ArrangedWidget.Geometry.AppendTransform(FSlateLayoutTransform(Inverse(ArrangedDesigner.Geometry.AbsolutePosition)));
-}
-
 #undef LOCTEXT_NAMESPACE
diff --git a/Engine/Source/Editor/UMGEditor/Private/Designer/DesignTimeUtils.h b/Engine/Source/Editor/UMGEditor/Private/Designer/DesignTimeUtils.h
--- a/Engine/Source/Editor/UMGEditor/Private/Designer/DesignTimeUtils.h
+++ b/Engine/Source/Editor/UMGEditor/Private/Designer/DesignTimeUtils.h
@@ -16,8 +16,4 @@ public:
 	static bool GetArrangedWidget(TSharedRef<SWidget> Widget, FArrangedWidget& ArrangedWidget);
 
 	static bool GetArrangedWidgetRelativeToWindow(TSharedRef<SWidget> Widget, FArrangedWidget& ArrangedWidget);
-
-	static bool GetArrangedWidgetRelativeToParent(TSharedRef<const SWidget> Widget, TSharedRef<const SWidget> Parent, FArrangedWidget& ArrangedWidget);
-
-	static void GetArrangedWidgetRelativeToParent(FWidgetPath& WidgetPath, TSharedRef<const SWidget> Widget, TSharedRef<const SWidget> Parent, FArrangedWidget& ArrangedWidget);
 };
diff --git a/Engine/Source/Editor/UMGEditor/Private/Designer/SDesignerView.cpp b/Engine/Source/Editor/UMGEditor/Private/Designer/SDesignerView.cpp
--- a/Engine/Source/Editor/UMGEditor/Private/Designer/SDesignerView.cpp
+++ b/Engine/Source/Editor/UMGEditor/Private/Designer/SDesignerView.cpp
@@ -1035,6 +1035,7 @@ FGeometry SDesignerView::MakeGeometryWindowLocal(const FGeometry& WidgetGeometry
 		TSharedRef<SWindow> CurrentWindowRef = WidgetWindow.ToSharedRef();
 
 		NewGeometry.AppendTransform(FSlateLayoutTransform(Inverse(CurrentWindowRef->GetPositionInScreen())));
+		//NewGeometry.AppendTransform(Inverse(CurrentWindowRef->GetLocalToScreenTransform()));
 	}
 
 	return NewGeometry;
@@ -1097,19 +1098,16 @@ FVector2D SDesignerView::GetExtensionPosition(TSharedRef<FDesignerSurfaceElement
 
 	if ( SelectedWidget.IsValid() )
 	{
-		FGeometry SelectedWidgetGeometry_RelativeToDesigner;
-		FGeometry SelectedWidgetParentGeometry_RelativeToDesigner;
+		FGeometry SelectedWidgetGeometry;
+		FGeometry SelectedWidgetParentGeometry;
 
-		if ( GetWidgetGeometry(SelectedWidget, SelectedWidgetGeometry_RelativeToDesigner) && GetWidgetParentGeometry(SelectedWidget, SelectedWidgetParentGeometry_RelativeToDesigner) )
+		if ( GetWidgetGeometry(SelectedWidget, SelectedWidgetGeometry) && GetWidgetParentGeometry(SelectedWidget, SelectedWidgetParentGeometry) )
 		{
-			SelectedWidgetGeometry_RelativeToDesigner.AppendTransform(FSlateLayoutTransform(Inverse(CachedGeometry.AbsolutePosition)));
-			SelectedWidgetParentGeometry_RelativeToDesigner.AppendTransform(FSlateLayoutTransform(Inverse(CachedGeometry.AbsolutePosition)));
+			const FVector2D WidgetPostion_DesignerSpace = (SelectedWidgetGeometry.AbsolutePosition - CachedGeometry.AbsolutePosition) / CachedGeometry.Scale;
+			const FVector2D WidgetSize = SelectedWidgetGeometry.Size * GetPreviewScale();
 
-			const FVector2D WidgetPostion_DesignerSpace = SelectedWidgetGeometry_RelativeToDesigner.AbsolutePosition;
-			const FVector2D WidgetSize = SelectedWidgetGeometry_RelativeToDesigner.Size * GetPreviewScale();
-
-			const FVector2D ParentPostion_DesignerSpace = SelectedWidgetParentGeometry_RelativeToDesigner.AbsolutePosition;
-			const FVector2D ParentSize = SelectedWidgetParentGeometry_RelativeToDesigner.Size * GetPreviewScale();
+			const FVector2D ParentPostion_DesignerSpace = (SelectedWidgetParentGeometry.AbsolutePosition - CachedGeometry.AbsolutePosition) / CachedGeometry.Scale;
+			const FVector2D ParentSize = SelectedWidgetParentGeometry.Size * GetPreviewScale();
 
 			FVector2D FinalPosition(0, 0);
 
diff --git a/Engine/Source/Editor/UMGEditor/Private/Designer/SZoomPan.cpp b/Engine/Source/Editor/UMGEditor/Private/Designer/SZoomPan.cpp
--- a/Engine/Source/Editor/UMGEditor/Private/Designer/SZoomPan.cpp
+++ b/Engine/Source/Editor/UMGEditor/Private/Designer/SZoomPan.cpp
@@ -43,7 +43,7 @@ void SZoomPan::OnArrangeChildren(const FGeometry& AllottedGeometry, FArrangedChi
 	}
 }
 
-float SZoomPan::GetRelativeLayoutScale(const FSlotBase& Child) const
+float SZoomPan::GetRelativeLayoutScale(const FSlotBase& Child, float LayoutScaleMultiplier) const
 {
 	return ZoomAmount.Get();
 }
diff --git a/Engine/Source/Editor/UMGEditor/Private/Designer/SZoomPan.h b/Engine/Source/Editor/UMGEditor/Private/Designer/SZoomPan.h
--- a/Engine/Source/Editor/UMGEditor/Private/Designer/SZoomPan.h
+++ b/Engine/Source/Editor/UMGEditor/Private/Designer/SZoomPan.h
@@ -35,7 +35,7 @@ public:
 protected:
 	void OnArrangeChildren(const FGeometry& AllottedGeometry, FArrangedChildren& ArrangedChildren) const override;
 
-	virtual float GetRelativeLayoutScale( const FSlotBase& Child ) const override;
+	virtual float GetRelativeLayoutScale( const FSlotBase& Child, float LayoutScaleMultiplier) const override;
 
 	/** The position within the panel at which the user is looking */
 	TAttribute<FVector2D> ViewOffset;
diff --git a/Engine/Source/Editor/UMGEditor/Private/Hierarchy/SHierarchyView.cpp b/Engine/Source/Editor/UMGEditor/Private/Hierarchy/SHierarchyView.cpp
--- a/Engine/Source/Editor/UMGEditor/Private/Hierarchy/SHierarchyView.cpp
+++ b/Engine/Source/Editor/UMGEditor/Private/Hierarchy/SHierarchyView.cpp
@@ -297,6 +297,13 @@ void SHierarchyView::RefreshTree()
 
 void SHierarchyView::RebuildTreeView()
 {
+	float OldScrollOffset = 0.0f;
+
+	if (WidgetTreeView.IsValid())
+	{
+		OldScrollOffset = WidgetTreeView->GetScrollOffset();
+	}
+
 	SAssignNew(WidgetTreeView, STreeView< TSharedPtr<FHierarchyModel> >)
 		.ItemHeight(20.0f)
 		.SelectionMode(ESelectionMode::Multi)
@@ -315,6 +322,9 @@ void SHierarchyView::RebuildTreeView()
 		[
 			WidgetTreeView.ToSharedRef()
 		]);
+
+	// Restore the previous scroll offset
+	WidgetTreeView->SetScrollOffset(OldScrollOffset);
 }
 
 void SHierarchyView::OnObjectsReplaced(const TMap<UObject*, UObject*>& ReplacementMap)
diff --git a/Engine/Source/Editor/UMGEditor/Private/Hierarchy/SHierarchyViewItem.cpp b/Engine/Source/Editor/UMGEditor/Private/Hierarchy/SHierarchyViewItem.cpp
--- a/Engine/Source/Editor/UMGEditor/Private/Hierarchy/SHierarchyViewItem.cpp
+++ b/Engine/Source/Editor/UMGEditor/Private/Hierarchy/SHierarchyViewItem.cpp
@@ -157,6 +157,7 @@ TOptional<EItemDropZone> ProcessHierarchyDragDrop(const FDragDropEvent& DragDrop
 	}
 
 	UWidgetBlueprint* Blueprint = BlueprintEditor->GetWidgetBlueprintObj();
+	check( Blueprint != nullptr && Blueprint->WidgetTree != nullptr );
 
 	// Is this a drag/drop op to create a new widget in the tree?
 	TSharedPtr<FWidgetTemplateDragDropOp> TemplateDragDropOp = DragDropEvent.GetOperationAs<FWidgetTemplateDragDropOp>();
@@ -330,8 +331,36 @@ TOptional<EItemDropZone> ProcessHierarchyDragDrop(const FDragDropEvent& DragDrop
 						}
 					}
 
+					UPanelWidget* OriginalParent = TemplateWidget->GetParent();
+					UBlueprint* OriginalBP = nullptr;
+
+					// The widget's parent is changing
+					if (OriginalParent != NewParent)
+					{
+						NewParent->SetFlags(RF_Transactional);
+						NewParent->Modify();
+
+						Blueprint->WidgetTree->SetFlags(RF_Transactional);
+						Blueprint->WidgetTree->Modify();
+
+						UWidgetTree* OriginalWidgetTree = Cast<UWidgetTree>(TemplateWidget->GetOuter());
+
+						if (OriginalWidgetTree != nullptr && UWidgetTree::TryMoveWidgetToNewTree(TemplateWidget, Blueprint->WidgetTree))
+						{
+							OriginalWidgetTree->SetFlags(RF_Transactional);
+							OriginalWidgetTree->Modify();
+
+							OriginalBP = OriginalWidgetTree->GetTypedOuter<UBlueprint>();
+						}
+					}
+
 					TemplateWidget->RemoveFromParent();
 
+					if (OriginalBP != nullptr && OriginalBP != Blueprint)
+					{
+						FBlueprintEditorUtils::MarkBlueprintAsStructurallyModified(OriginalBP);
+					}
+
 					UPanelSlot* NewSlot = nullptr;
 					if (Index.IsSet())
 					{
diff --git a/Engine/Source/Editor/UMGEditor/Private/WidgetBlueprint.cpp b/Engine/Source/Editor/UMGEditor/Private/WidgetBlueprint.cpp
--- a/Engine/Source/Editor/UMGEditor/Private/WidgetBlueprint.cpp
+++ b/Engine/Source/Editor/UMGEditor/Private/WidgetBlueprint.cpp
@@ -16,6 +16,15 @@
 
 #include "Kismet2/CompilerResultsLog.h"
 #include "Binding/PropertyBinding.h"
+#include "Blueprint/WidgetTree.h"
+#include "Blueprint/WidgetBlueprintGeneratedClass.h"
+#include "PropertyTag.h"
+#include "WidgetBlueprint.h"
+#include "WidgetBlueprintCompiler.h"
+#include "PropertyBinding.h"
+#include "Engine/UserDefinedStruct.h"
+#include "UObject/EditorObjectVersion.h"
+#include "Classes/WidgetGraphSchema.h"
 
 #define LOCTEXT_NAMESPACE "UMG"
 
@@ -504,6 +513,23 @@ UWidgetBlueprint::UWidgetBlueprint(const FObjectInitializer& ObjectInitializer)
 	WidgetTree->SetFlags(RF_Transactional);
 }
 
+void UWidgetBlueprint::ReplaceDeprecatedNodes()
+{
+	if ( GetLinkerCustomVersion(FEditorObjectVersion::GUID) < FEditorObjectVersion::WidgetGraphSchema )
+	{
+		// Update old graphs to all use the widget graph schema.
+		TArray<UEdGraph*> Graphs;
+		GetAllGraphs(Graphs);
+
+		for ( UEdGraph* Graph : Graphs )
+		{
+			Graph->Schema = UWidgetGraphSchema::StaticClass();
+		}
+	}
+
+	Super::ReplaceDeprecatedNodes();
+}
+
 void UWidgetBlueprint::PostLoad()
 {
 	Super::PostLoad();
@@ -548,6 +574,18 @@ void UWidgetBlueprint::PostLoad()
 			}
 		}
 	}
+
+	if ( GetLinkerCustomVersion(FEditorObjectVersion::GUID) < FEditorObjectVersion::WidgetGraphSchema )
+	{
+		// Update old graphs to all use the widget graph schema.
+		TArray<UEdGraph*> Graphs;
+		GetAllGraphs(Graphs);
+
+		for ( UEdGraph* Graph : Graphs )
+		{
+			Graph->Schema = UWidgetGraphSchema::StaticClass();
+		}
+	}
 }
 
 void UWidgetBlueprint::PostDuplicate(bool bDuplicateForPIE)
diff --git a/Engine/Source/Editor/UMGEditor/Private/WidgetBlueprintCompiler.cpp b/Engine/Source/Editor/UMGEditor/Private/WidgetBlueprintCompiler.cpp
--- a/Engine/Source/Editor/UMGEditor/Private/WidgetBlueprintCompiler.cpp
+++ b/Engine/Source/Editor/UMGEditor/Private/WidgetBlueprintCompiler.cpp
@@ -15,11 +15,14 @@
 #include "Kismet2/BlueprintEditorUtils.h"
 #include "Components/NamedSlot.h"
 #include "WidgetBlueprintEditorUtils.h"
+#include "WidgetGraphSchema.h"
 
 #define LOCTEXT_NAMESPACE "UMG"
 
 FWidgetBlueprintCompiler::FWidgetBlueprintCompiler(UWidgetBlueprint* SourceSketch, FCompilerResultsLog& InMessageLog, const FKismetCompilerOptions& InCompilerOptions, TArray<UObject*>* InObjLoaded)
 	: Super(SourceSketch, InMessageLog, InCompilerOptions, InObjLoaded)
+	, NewWidgetBlueprintClass(nullptr)
+	, WidgetSchema(nullptr)
 {
 }
 
@@ -27,6 +30,12 @@ FWidgetBlueprintCompiler::~FWidgetBlueprintCompiler()
 {
 }
 
+UEdGraphSchema_K2* FWidgetBlueprintCompiler::CreateSchema()
+{
+	WidgetSchema = NewObject<UWidgetGraphSchema>();
+	return WidgetSchema;
+}
+
 void FWidgetBlueprintCompiler::CreateFunctionList()
 {
 	Super::CreateFunctionList();
diff --git a/Engine/Source/Editor/UMGEditor/Private/WidgetBlueprintCompiler.h b/Engine/Source/Editor/UMGEditor/Private/WidgetBlueprintCompiler.h
--- a/Engine/Source/Editor/UMGEditor/Private/WidgetBlueprintCompiler.h
+++ b/Engine/Source/Editor/UMGEditor/Private/WidgetBlueprintCompiler.h
@@ -6,6 +6,7 @@
 #include "WidgetBlueprint.h"
 #include "KismetCompiler.h"
 
+
 //////////////////////////////////////////////////////////////////////////
 // FWidgetBlueprintCompiler
 
@@ -28,7 +29,7 @@ protected:
 	void ValidateWidgetNames();
 
 	// FKismetCompilerContext
-	//virtual UEdGraphSchema_K2* CreateSchema() override;
+	virtual UEdGraphSchema_K2* CreateSchema() override;
 	virtual void CreateFunctionList() override;
 	virtual void SpawnNewClass(const FString& NewClassName) override;
 	virtual void PrecompileFunction(FKismetFunctionContext& Context) override;
@@ -45,6 +46,8 @@ protected:
 protected:
 	UWidgetBlueprintGeneratedClass* NewWidgetBlueprintClass;
 
+	UWidgetGraphSchema* WidgetSchema;
+
 	// Map of properties created for widgets; to aid in debug data generation
 	TMap<class UWidget*, class UProperty*> WidgetToMemberVariableMap;
 
diff --git a/Engine/Source/Editor/UMGEditor/Private/WidgetBlueprintEditor.cpp b/Engine/Source/Editor/UMGEditor/Private/WidgetBlueprintEditor.cpp
--- a/Engine/Source/Editor/UMGEditor/Private/WidgetBlueprintEditor.cpp
+++ b/Engine/Source/Editor/UMGEditor/Private/WidgetBlueprintEditor.cpp
@@ -36,6 +36,7 @@
 #include "Framework/Commands/GenericCommands.h"
 #include "Kismet2/CompilerResultsLog.h"
 #include "IMessageLogListing.h"
+#include "WidgetGraphSchema.h"
 
 #include "Animation/MovieSceneWidgetMaterialTrack.h"
 #include "Animation/WidgetMaterialTrackUtilities.h"
@@ -116,14 +117,6 @@ void FWidgetBlueprintEditor::InitWidgetBlueprintEditor(const EToolkitMode::Type
 	// for change selected widgets on sequencer tree view
 	UWidgetBlueprint* Blueprint = GetWidgetBlueprintObj();
 
-	// If this blueprint is empty, add a canvas panel as the root widget.
-	if ( Blueprint->WidgetTree->RootWidget == nullptr )
-	{
-		UWidget* RootWidget = Blueprint->WidgetTree->ConstructWidget<UCanvasPanel>(UCanvasPanel::StaticClass());
-		RootWidget->SetDesignerFlags(GetCurrentDesignerFlags());
-		Blueprint->WidgetTree->RootWidget = RootWidget;
-	}
-
 	UpdatePreview(GetWidgetBlueprintObj(), true);
 
 	DesignerCommandList = MakeShareable(new FUICommandList);
@@ -881,6 +874,11 @@ FGraphAppearanceInfo FWidgetBlueprintEditor::GetGraphAppearance(UEdGraph* InGrap
 	return AppearanceInfo;
 }
 
+TSubclassOf<UEdGraphSchema> FWidgetBlueprintEditor::GetDefaultSchemaClass() const
+{
+	return UWidgetGraphSchema::StaticClass();
+}
+
 void FWidgetBlueprintEditor::ClearHoveredWidget()
 {
 	HoveredWidget = FWidgetReference();
@@ -957,52 +955,31 @@ public:
 
 };
 
-void GetBindableObjects(UWidget* RootWidget, TArray<FObjectAndDisplayName>& BindableObjects)
+void GetBindableObjects(UWidgetTree* WidgetTree, TArray<FObjectAndDisplayName>& BindableObjects)
 {
-	TArray<UWidget*> ToTraverse;
-	ToTraverse.Add(RootWidget);
-	while (ToTraverse.Num() > 0)
-	{
-		UWidget* Widget = ToTraverse[0];
-		ToTraverse.RemoveAt(0);
+	WidgetTree->ForEachWidget([&BindableObjects] (UWidget* Widget) {
 		BindableObjects.Add(FObjectAndDisplayName(FText::FromString(Widget->GetName()), Widget));
 
-		UUserWidget* UserWidget = Cast<UUserWidget>(Widget);
-		if (UserWidget != nullptr)
-		{
-			TArray<FName> SlotNames;
-			UserWidget->GetSlotNames(SlotNames);
-			for (FName SlotName : SlotNames)
-			{
-				UWidget* Content = UserWidget->GetContentForSlot(SlotName);
-				if (Content != nullptr)
-				{
-					ToTraverse.Add(Content);
-				}
-			}
-		}
-
 		UPanelWidget* PanelWidget = Cast<UPanelWidget>(Widget);
-		if (PanelWidget != nullptr)
+		if ( PanelWidget != nullptr )
 		{
-			for (UPanelSlot* Slot : PanelWidget->GetSlots())
+			for ( UPanelSlot* Slot : PanelWidget->GetSlots() )
 			{
-				if (Slot->Content != nullptr)
+				if ( Slot->Content != nullptr )
 				{
 					FText SlotDisplayName = FText::Format(LOCTEXT("AddMenuSlotFormat", "{0} ({1} Slot)"), FText::FromString(Slot->Content->GetName()), FText::FromString(PanelWidget->GetName()));
 					BindableObjects.Add(FObjectAndDisplayName(SlotDisplayName, Slot));
-					ToTraverse.Add(Slot->Content);
 				}
 			}
 		}
-	}
+	});
 }
 
 void FWidgetBlueprintEditor::OnGetAnimationAddMenuContent(FMenuBuilder& MenuBuilder, TSharedRef<ISequencer> InSequencer)
 {
 	TArray<FObjectAndDisplayName> BindableObjects;
 	{
-		GetBindableObjects(GetPreview()->GetRootWidget(), BindableObjects);
+		GetBindableObjects(GetPreview()->WidgetTree, BindableObjects);
 		BindableObjects.Sort();
 	}
 
diff --git a/Engine/Source/Editor/UMGEditor/Private/WidgetBlueprintFactory.cpp b/Engine/Source/Editor/UMGEditor/Private/WidgetBlueprintFactory.cpp
--- a/Engine/Source/Editor/UMGEditor/Private/WidgetBlueprintFactory.cpp
+++ b/Engine/Source/Editor/UMGEditor/Private/WidgetBlueprintFactory.cpp
@@ -10,6 +10,8 @@
 
 
 
+#include "Components/CanvasPanel.h"
+
 #define LOCTEXT_NAMESPACE "UWidgetBlueprintFactory"
 
 /*------------------------------------------------------------------------------
@@ -58,6 +60,13 @@ UObject* UWidgetBlueprintFactory::FactoryCreateNew(UClass* Class, UObject* InPar
 	{
 		UWidgetBlueprint* NewBP = CastChecked<UWidgetBlueprint>(FKismetEditorUtilities::CreateBlueprint(ParentClass, InParent, Name, BlueprintType, UWidgetBlueprint::StaticClass(), UWidgetBlueprintGeneratedClass::StaticClass(), CallingContext));
 
+		// Create a CanvasPanel to use as the default root widget
+		if ( NewBP->WidgetTree->RootWidget == nullptr )
+		{
+			UWidget* Root = NewBP->WidgetTree->ConstructWidget<UCanvasPanel>(UCanvasPanel::StaticClass());
+			NewBP->WidgetTree->RootWidget = Root;
+		}
+
 		return NewBP;
 	}
 }
diff --git a/Engine/Source/Editor/UMGEditor/Private/WidgetGraphSchema.cpp b/Engine/Source/Editor/UMGEditor/Private/WidgetGraphSchema.cpp
--- a/Engine/Source/Editor/UMGEditor/Private/WidgetGraphSchema.cpp
+++ b/Engine/Source/Editor/UMGEditor/Private/WidgetGraphSchema.cpp
@@ -5,10 +5,54 @@
 UWidgetGraphSchema::UWidgetGraphSchema(const FObjectInitializer& ObjectInitializer)
 : Super(ObjectInitializer)
 {
-	NAME_NeverAsPin = TEXT("NeverAsPin");
-	NAME_PinHiddenByDefault = TEXT("PinHiddenByDefault");
-	NAME_PinShownByDefault = TEXT("PinShownByDefault");
-	NAME_AlwaysAsPin = TEXT("AlwaysAsPin");
-	NAME_OnEvaluate = TEXT("OnEvaluate");
-	DefaultEvaluationHandlerName = TEXT("EvaluateGraphExposedInputs");
+}
+
+void UWidgetGraphSchema::BackwardCompatibilityNodeConversion(UEdGraph* Graph, bool bOnlySafeChanges) const
+{
+	Super::BackwardCompatibilityNodeConversion(Graph, bOnlySafeChanges);
+
+	//if ( Graph )
+	//{
+	//	{
+	//		static const FString BlueprintPinName(TEXT("Blueprint"));
+	//		static const FString ClassPinName(TEXT("Class"));
+	//		TArray<UK2Node_SpawnActor*> SpawnActorNodes;
+	//		Graph->GetNodesOfClass(SpawnActorNodes);
+	//		for ( auto It = SpawnActorNodes.CreateIterator(); It; ++It )
+	//		{
+	//			FBackwardCompatibilityConversionHelper::ConvertNode(
+	//				*It, BlueprintPinName, NewObject<UK2Node_SpawnActorFromClass>(Graph),
+	//				ClassPinName, *this, bOnlySafeChanges);
+	//		}
+	//	}
+
+	//	{
+	//		auto Blueprint = FBlueprintEditorUtils::FindBlueprintForGraph(Graph);
+	//		if ( Blueprint && *Blueprint->SkeletonGeneratedClass )
+	//		{
+	//			TArray<UK2Node_CallFunction*> Nodes;
+	//			Graph->GetNodesOfClass(Nodes);
+	//			for ( const auto& FunctionRedirect : EditoronlyBPFunctionRedirects )
+	//			{
+	//				FBackwardCompatibilityConversionHelper::ConvertFunctionCallNodes(
+	//					FBackwardCompatibilityConversionHelper::FFunctionCallParams(FunctionRedirect),
+	//					Nodes, Graph, *this, bOnlySafeChanges);
+	//			}
+	//		}
+	//		else
+	//		{
+	//			UE_LOG(LogBlueprint, Log, TEXT("BackwardCompatibilityNodeConversion: Blueprint '%s' cannot be fully converted. It has no skeleton class!"),
+	//				Blueprint ? *Blueprint->GetName() : TEXT("Unknown"));
+	//		}
+	//	}
+
+	//	// Call per-node deprecation functions
+	//	TArray<UK2Node*> PossiblyDeprecatedNodes;
+	//	Graph->GetNodesOfClass<UK2Node>(PossiblyDeprecatedNodes);
+
+	//	for ( UK2Node* Node : PossiblyDeprecatedNodes )
+	//	{
+	//		Node->ConvertDeprecatedNode(Graph, bOnlySafeChanges);
+	//	}
+	//}
 }
diff --git a/Engine/Source/Editor/UMGEditor/Public/WidgetBlueprint.h b/Engine/Source/Editor/UMGEditor/Public/WidgetBlueprint.h
--- a/Engine/Source/Editor/UMGEditor/Public/WidgetBlueprint.h
+++ b/Engine/Source/Editor/UMGEditor/Public/WidgetBlueprint.h
@@ -211,6 +211,8 @@ public:
 	/** UObject interface */
 	virtual void PostLoad() override;
 	virtual void PostDuplicate(bool bDuplicateForPIE) override;
+
+	virtual void ReplaceDeprecatedNodes() override;
 	
 	//~ Begin UBlueprint Interface
 	virtual UClass* GetBlueprintClass() const override;
diff --git a/Engine/Source/Editor/UnrealEd/Classes/ActorFactories/ActorFactoryBasicShape.h b/Engine/Source/Editor/UnrealEd/Classes/ActorFactories/ActorFactoryBasicShape.h
--- a/Engine/Source/Editor/UnrealEd/Classes/ActorFactories/ActorFactoryBasicShape.h
+++ b/Engine/Source/Editor/UnrealEd/Classes/ActorFactories/ActorFactoryBasicShape.h
@@ -23,4 +23,5 @@ class UActorFactoryBasicShape : public UActorFactoryStaticMesh
 	UNREALED_API static const FName BasicSphere;
 	UNREALED_API static const FName BasicCylinder;
 	UNREALED_API static const FName BasicCone;
+	UNREALED_API static const FName BasicPlane;
 };
diff --git a/Engine/Source/Editor/UnrealEd/Classes/Editor/EditorEngine.h b/Engine/Source/Editor/UnrealEd/Classes/Editor/EditorEngine.h
--- a/Engine/Source/Editor/UnrealEd/Classes/Editor/EditorEngine.h
+++ b/Engine/Source/Editor/UnrealEd/Classes/Editor/EditorEngine.h
@@ -2850,6 +2850,9 @@ private:
 	bool bPlayUsingLauncherHasCode;
 	bool bPlayUsingLauncherHasCompiler;
 
+	/** Used to prevent reentrant calls to EndPlayMap(). */
+	bool bIsEndingPlay;
+
 	/** List of files we are deferring adding to source control */
 	TArray<FString> DeferredFilesToAddToSourceControl;
 
@@ -2911,6 +2914,9 @@ protected:
 	// Handle requests from slate application to open assets.
 	bool HandleOpenAsset(UObject* Asset);
 
+	// Handles a package being reloaded.
+	void HandlePackageReloaded(const EPackageReloadPhase InPackageReloadPhase, FPackageReloadedEvent* InPackageReloadedEvent);
+
 public:
 	/** True if world assets are enabled */
 	static bool IsUsingWorldAssets();
diff --git a/Engine/Source/Editor/UnrealEd/Classes/Editor/PropertyEditorTestObject.h b/Engine/Source/Editor/UnrealEd/Classes/Editor/PropertyEditorTestObject.h
--- a/Engine/Source/Editor/UnrealEd/Classes/Editor/PropertyEditorTestObject.h
+++ b/Engine/Source/Editor/UnrealEd/Classes/Editor/PropertyEditorTestObject.h
@@ -267,7 +267,10 @@ class UPropertyEditorTestObject : public UObject
 	UPROPERTY(EditAnywhere, Category=StructTests)
 	FPropertyEditorTestBasicStruct StructWithMultipleInstances1;
 
-	UPROPERTY(EditAnywhere, Category=StructTests, meta=(editcondition = "bEditCondition"))
+	UPROPERTY(EditAnywhere, Category = StructTests, meta = (InlineEditConditionToggle))
+	bool bEditConditionStructWithMultipleInstances2;
+
+	UPROPERTY(EditAnywhere, Category=StructTests, meta=(editcondition = "bEditConditionStructWithMultipleInstances2"))
 	FPropertyEditorTestBasicStruct StructWithMultipleInstances2;
 
 	UPROPERTY(EditAnywhere, Category=StructTests)
@@ -282,7 +285,10 @@ class UPropertyEditorTestObject : public UObject
 	UPROPERTY(EditAnywhere, Category=AdvancedProperties, meta=(editcondition = "bEditCondition"))
 	int32 SimplePropertyWithEditCondition;
 
-	UPROPERTY(EditAnywhere, Category=StructTests, meta=(editcondition = "bEditCondition"))
+	UPROPERTY(EditAnywhere, Category = StructTests, meta = (InlineEditConditionToggle))
+	bool bEditConditionAssetReferenceCustomStructWithEditCondition;
+
+	UPROPERTY(EditAnywhere, Category=StructTests, meta=(editcondition = "bEditConditionAssetReferenceCustomStructWithEditCondition"))
 	FStringAssetReference AssetReferenceCustomStructWithEditCondition;
 
 	UPROPERTY(EditAnywhere, Category=StructTests)
diff --git a/Engine/Source/Editor/UnrealEd/Classes/Editor/UnrealEdEngine.h b/Engine/Source/Editor/UnrealEd/Classes/Editor/UnrealEdEngine.h
--- a/Engine/Source/Editor/UnrealEd/Classes/Editor/UnrealEdEngine.h
+++ b/Engine/Source/Editor/UnrealEd/Classes/Editor/UnrealEdEngine.h
@@ -156,6 +156,9 @@ public:
 	UPROPERTY()
 	uint32 bNeedWarningForWritePermission:1;
 
+	/** Whether there is a pending package notification */
+ 	uint32 bShowPackageNotification:1;
+
 	/** Array of sorted, localized editor sprite categories */
 	UPROPERTY()
 	TArray<FString> SortedSpriteCategories_DEPRECATED;
diff --git a/Engine/Source/Editor/UnrealEd/Classes/Factories/FbxAssetImportData.h b/Engine/Source/Editor/UnrealEd/Classes/Factories/FbxAssetImportData.h
--- a/Engine/Source/Editor/UnrealEd/Classes/Factories/FbxAssetImportData.h
+++ b/Engine/Source/Editor/UnrealEd/Classes/Factories/FbxAssetImportData.h
@@ -26,6 +26,18 @@ class UNREALED_API UFbxAssetImportData : public UAssetImportData
 	UPROPERTY(EditAnywhere, Category=Transform, meta=(ImportType="StaticMesh|SkeletalMesh|Animation", ImportCategory="Transform"))
 	float ImportUniformScale;
 
+	/** Whether to convert scene from FBX scene. */
+	UPROPERTY(EditAnywhere, config, Category = Miscellaneous, meta = (ImportType = "StaticMesh|SkeletalMesh|Animation", ImportCategory = "Miscellaneous", ToolTip = "Convert the scene from FBX coordinate system to UE4 coordinate system"))
+	bool bConvertScene;
+
+	/** Whether to force the front axis to be align with X instead of -Y. */
+	UPROPERTY(EditAnywhere, config, Category = Miscellaneous, meta = (editcondition = "bConvertScene", ImportType = "StaticMesh|SkeletalMesh|Animation", ImportCategory = "Miscellaneous", ToolTip = "Convert the scene from FBX coordinate system to UE4 coordinate system with front X axis instead of -Y"))
+	bool bForceFrontXAxis;
+
+	/** Whether to convert the scene from FBX unit to UE4 unit (centimeter). */
+	UPROPERTY(EditAnywhere, config, Category = Miscellaneous, meta = (ImportType = "StaticMesh|SkeletalMesh|Animation", ImportCategory = "Miscellaneous", ToolTip = "Convert the scene from FBX unit to UE4 unit (centimeter)."))
+	bool bConvertSceneUnit;
+
 	/* Use by the reimport factory to answer CanReimport, if true only factory for scene reimport will return true */
 	UPROPERTY()
 	bool bImportAsScene;
diff --git a/Engine/Source/Editor/UnrealEd/Classes/Factories/FbxImportUI.h b/Engine/Source/Editor/UnrealEd/Classes/Factories/FbxImportUI.h
--- a/Engine/Source/Editor/UnrealEd/Classes/Factories/FbxImportUI.h
+++ b/Engine/Source/Editor/UnrealEd/Classes/Factories/FbxImportUI.h
@@ -50,18 +50,6 @@ public:
 	UPROPERTY(EditAnywhere, AdvancedDisplay, config, Category=Miscellaneous, meta=(OBJRestrict="true"))
 	uint32 bOverrideFullName:1;
 
-	/** Whether to convert scene from FBX scene. */
-	UPROPERTY(EditAnywhere, AdvancedDisplay, config, Category=Miscellaneous, meta=(OBJRestrict="true", ToolTip="Convert the scene from FBX coordinate system to UE4 coordinate system"))
-	uint32 bConvertScene:1;
-	
-	/** Whether to force the front axis to be align with X instead of -Y. */
-	UPROPERTY(EditAnywhere, AdvancedDisplay, config, Category = Miscellaneous, meta = (editcondition = "bConvertScene", OBJRestrict = "true", ToolTip = "Convert the scene from FBX coordinate system to UE4 coordinate system with front X axis instead of -Y"))
-	uint32 bForceFrontXAxis : 1;
-
-	/** Whether to convert the scene from FBX unit to UE4 unit (centimeter). */
-	UPROPERTY(EditAnywhere, AdvancedDisplay, config, Category = Miscellaneous, meta = (OBJRestrict = "true", ToolTip = "Convert the scene from FBX unit to UE4 unit (centimeter)."))
-	uint32 bConvertSceneUnit : 1;
-
 	/** Whether to import the incoming FBX as a skeletal object */
 	UPROPERTY(EditAnywhere, Category = Mesh, meta = (ImportType = "StaticMesh|SkeletalMesh", DisplayName="Skeletal Mesh"))
 	bool bImportAsSkeletal;
@@ -74,10 +62,6 @@ public:
 	UPROPERTY(EditAnywhere, Category=Mesh, meta=(ImportType="SkeletalMesh"))
 	bool bImportMesh;
 
-	/** For static meshes, enabling this option will combine all meshes in the FBX into a single monolithic mesh in Unreal */
-	UPROPERTY(EditAnywhere, AdvancedDisplay, config, Category=Mesh, meta=(ToolTip="If enabled, combines all meshes into a single mesh", ImportType="StaticMesh"))
-	uint32 bCombineMeshes:1;
-
 	/** Skeleton to use for imported asset. When importing a mesh, leaving this as "None" will create a new skeleton. When importing an animation this MUST be specified to import the asset. */
 	UPROPERTY(EditAnywhere, Category=Mesh, meta=(ImportType="SkeletalMesh|Animation"))
 	class USkeleton* Skeleton;
@@ -111,21 +95,25 @@ public:
 	uint32 bImportTextures:1;
 
 	/** Import data used when importing static meshes */
-	UPROPERTY(EditAnywhere, Instanced, Category = Mesh, meta=(ImportType = "StaticMesh"))
+	UPROPERTY(EditAnywhere, Transient, Instanced, Category = Mesh, meta=(ImportType = "StaticMesh"))
 	class UFbxStaticMeshImportData* StaticMeshImportData;
 
 	/** Import data used when importing skeletal meshes */
-	UPROPERTY(EditAnywhere, Instanced, Category=Mesh, meta=(ImportType = "SkeletalMesh"))
+	UPROPERTY(EditAnywhere, Transient, Instanced, Category=Mesh, meta=(ImportType = "SkeletalMesh"))
 	class UFbxSkeletalMeshImportData* SkeletalMeshImportData;
 
 	/** Import data used when importing animations */
-	UPROPERTY(EditAnywhere, Instanced, Category=Animation, meta=(editcondition="bImportAnimations", ImportType = "Animation"))
+	UPROPERTY(EditAnywhere, Transient, Instanced, Category=Animation, meta=(editcondition="bImportAnimations", ImportType = "Animation"))
 	class UFbxAnimSequenceImportData* AnimSequenceImportData;
 
 	/** Import data used when importing textures */
-	UPROPERTY(EditAnywhere, Instanced, Category=Material)
+	UPROPERTY(EditAnywhere, Transient, Instanced, Category=Material)
 	class UFbxTextureImportData* TextureImportData;
 
+	/** If true the automated import path should detect the import type.  If false the import type was specified by the user */
+	UPROPERTY()
+	bool bAutomatedImportShouldDetectType;
+
 	/** UObject Interface */
 	virtual bool CanEditChange( const UProperty* InProperty ) const override;
 
diff --git a/Engine/Source/Editor/UnrealEd/Classes/Factories/FbxSceneImportData.h b/Engine/Source/Editor/UnrealEd/Classes/Factories/FbxSceneImportData.h
--- a/Engine/Source/Editor/UnrealEd/Classes/Factories/FbxSceneImportData.h
+++ b/Engine/Source/Editor/UnrealEd/Classes/Factories/FbxSceneImportData.h
@@ -43,6 +43,9 @@ public:
 	/* Is the original import create a folder hierarchy */
 	bool bCreateFolderHierarchy;
 
+	/* Is the original import force front X axis */
+	bool bForceFrontXAxis;
+
 	/* Which type of hierarchy was create see */
 	int32 HierarchyType;
 
diff --git a/Engine/Source/Editor/UnrealEd/Classes/Factories/FbxSceneImportFactory.h b/Engine/Source/Editor/UnrealEd/Classes/Factories/FbxSceneImportFactory.h
--- a/Engine/Source/Editor/UnrealEd/Classes/Factories/FbxSceneImportFactory.h
+++ b/Engine/Source/Editor/UnrealEd/Classes/Factories/FbxSceneImportFactory.h
@@ -182,7 +182,7 @@ public:
 	float FarPlane;
 	float OrthoZoom;
 	float FieldOfView;
-	
+	float FocalLength;
 
 	FFbxCameraInfo()
 		: Name(TEXT(""))
@@ -194,6 +194,7 @@ public:
 		, FarPlane(0.0f)
 		, OrthoZoom(0.0f)
 		, FieldOfView(0.0f)
+		, FocalLength(0.0f)
 	{}
 };
 
@@ -378,6 +379,11 @@ class UNREALED_API UFbxSceneImportFactory : public UFactory
 	static FString DefaultOptionName;
 
 protected:
+	/** Convert the scene and remake all the transform for the SceneInfo pass in parameter.
+	 *  We need this because EvaluateGlobal and EvaluateLocal are dependent of the scene conversion.
+	 */
+	void ChangeFrontAxis(void* VoidFbxImporter, void* VoidSceneInfo, TSharedPtr<FFbxSceneInfo> SceneInfoPtr);
+
 	/** Make sure GlobalImportSettings is pointing to the correct options */
 	void ApplyMeshInfoFbxOptions(TSharedPtr<FFbxMeshInfo> MeshInfo);
 
diff --git a/Engine/Source/Editor/UnrealEd/Classes/Factories/FbxSceneImportOptions.h b/Engine/Source/Editor/UnrealEd/Classes/Factories/FbxSceneImportOptions.h
--- a/Engine/Source/Editor/UnrealEd/Classes/Factories/FbxSceneImportOptions.h
+++ b/Engine/Source/Editor/UnrealEd/Classes/Factories/FbxSceneImportOptions.h
@@ -46,6 +46,10 @@ class UFbxSceneImportOptions : public UObject
 	UPROPERTY(EditAnywhere, config, category = ImportOptions)
 	TEnumAsByte<enum EFBXSceneOptionsCreateHierarchyType> HierarchyType;
 
+	/** Whether to force the front axis to be align with X instead of -Y. */
+	UPROPERTY(EditAnywhere, config, category = ImportOptions)
+	uint32 bForceFrontXAxis : 1;
+
 	UPROPERTY()
 	FVector ImportTranslation;
 
diff --git a/Engine/Source/Editor/UnrealEd/Classes/Factories/FbxStaticMeshImportData.h b/Engine/Source/Editor/UnrealEd/Classes/Factories/FbxStaticMeshImportData.h
--- a/Engine/Source/Editor/UnrealEd/Classes/Factories/FbxStaticMeshImportData.h
+++ b/Engine/Source/Editor/UnrealEd/Classes/Factories/FbxStaticMeshImportData.h
@@ -66,6 +66,10 @@ class UFbxStaticMeshImportData : public UFbxMeshImportData
 	UPROPERTY(EditAnywhere, config, Category = ImportSettings, meta=(OBJRestrict="true", ImportType="StaticMesh"))
 	uint32 bAutoGenerateCollision : 1;
 
+	/** For static meshes, enabling this option will combine all meshes in the FBX into a single monolithic mesh in Unreal */
+	UPROPERTY(EditAnywhere, AdvancedDisplay, config, Category = Mesh, meta = (ToolTip = "If enabled, combines all meshes into a single mesh", ImportType = "StaticMesh"))
+	uint32 bCombineMeshes : 1;
+
 	/** Gets or creates fbx import data for the specified static mesh */
 	static UFbxStaticMeshImportData* GetImportDataForStaticMesh(UStaticMesh* StaticMesh, UFbxStaticMeshImportData* TemplateForCreation);
 
diff --git a/Engine/Source/Editor/UnrealEd/Classes/Factories/FbxTextureImportData.h b/Engine/Source/Editor/UnrealEd/Classes/Factories/FbxTextureImportData.h
--- a/Engine/Source/Editor/UnrealEd/Classes/Factories/FbxTextureImportData.h
+++ b/Engine/Source/Editor/UnrealEd/Classes/Factories/FbxTextureImportData.h
@@ -8,6 +8,22 @@
 #include "Factories/FbxAssetImportData.h"
 #include "FbxTextureImportData.generated.h"
 
+UENUM()
+namespace EMaterialSearchLocation
+{
+	enum Type
+	{
+		/** Search for matching material in local import folder only. */
+		Local,
+		/** Search for matching material recursively from parent folder. */
+		UnderParent,
+		/** Search for matching material recursively from root folder. */
+		UnderRoot,
+		/** Search for matching material in all assets folders. */
+		AllAssets
+	};
+}
+
 /**
  * Import data and options used when importing any mesh from FBX
  */
@@ -20,6 +36,10 @@ class UNREALED_API UFbxTextureImportData : public UFbxAssetImportData
 	UPROPERTY(EditAnywhere, AdvancedDisplay, config, Category=ImportSettings, meta=(OBJRestrict="true"))
 	uint32 bInvertNormalMaps:1;
 
+	/** Specify where we should search for matching materials when importing */
+	UPROPERTY(EditAnywhere, config, Category = ImportSettings, meta = (DisplayName="Search Location", OBJRestrict = "true", ImportType = "Mesh"))
+	TEnumAsByte<EMaterialSearchLocation::Type> MaterialSearchLocation;
+
 	/** Base material to instance from when importing materials. */
 	UPROPERTY(EditAnywhere, config, Category = Material, meta = (ImportType = "Mesh", AllowedClasses = "MaterialInterface"))
 	FStringAssetReference BaseMaterialName;
diff --git a/Engine/Source/Editor/UnrealEd/Classes/MaterialGraph/MaterialGraphNode.h b/Engine/Source/Editor/UnrealEd/Classes/MaterialGraph/MaterialGraphNode.h
--- a/Engine/Source/Editor/UnrealEd/Classes/MaterialGraph/MaterialGraphNode.h
+++ b/Engine/Source/Editor/UnrealEd/Classes/MaterialGraph/MaterialGraphNode.h
@@ -78,6 +78,9 @@ public:
 	virtual uint32 GetInputType(const UEdGraphPin* InputPin) const override;
 	//~ End UMaterialGraphNode_Base Interface
 
+	/** Will return the shorten pin name to use based on long pin name */
+	FString UNREALED_API GetShortenPinName(const FString& PinName);
+
 private:
 	/** Make sure the MaterialExpression is owned by the Material */
 	void ResetMaterialExpressionOwner();
diff --git a/Engine/Source/Editor/UnrealEd/Classes/MaterialGraph/MaterialGraphSchema.h b/Engine/Source/Editor/UnrealEd/Classes/MaterialGraph/MaterialGraphSchema.h
--- a/Engine/Source/Editor/UnrealEd/Classes/MaterialGraph/MaterialGraphSchema.h
+++ b/Engine/Source/Editor/UnrealEd/Classes/MaterialGraph/MaterialGraphSchema.h
@@ -219,6 +219,8 @@ class UMaterialGraphSchema : public UEdGraphSchema
 	virtual bool IsCacheVisualizationOutOfDate(int32 InVisualizationCacheID) const override;
 	virtual int32 GetCurrentVisualizationCacheID() const override;
 	virtual void ForceVisualizationCacheClear() const override;
+	virtual void OnPinConnectionDoubleCicked(UEdGraphPin* PinA, UEdGraphPin* PinB, const FVector2D& GraphPosition) const override;
+	virtual bool SafeDeleteNodeFromGraph(UEdGraph* Graph, UEdGraphNode* NodeToDelete) const override;
 	//~ End UEdGraphSchema Interface
 
 private:
diff --git a/Engine/Source/Editor/UnrealEd/Classes/Settings/EditorExperimentalSettings.h b/Engine/Source/Editor/UnrealEd/Classes/Settings/EditorExperimentalSettings.h
--- a/Engine/Source/Editor/UnrealEd/Classes/Settings/EditorExperimentalSettings.h
+++ b/Engine/Source/Editor/UnrealEd/Classes/Settings/EditorExperimentalSettings.h
@@ -42,6 +42,10 @@ public:
 	UPROPERTY(EditAnywhere, config, Category = Tools, meta = (DisplayName = "Enable Details Panel Favorites"))
 	bool bEnableFavoriteSystem;
 
+	/** Enables content hot-reloading in the editor (eg, when syncing new assets via source control) */
+	UPROPERTY(EditAnywhere, config, Category = Tools, meta = (DisplayName = "Content Hot-Reloading"))
+	bool bEnableContentHotReloading;
+
 	/** Enable being able to subclass components in blueprints */
 	UPROPERTY(EditAnywhere, config, Category=Tools, meta=(ConfigRestartRequired=true))
 	bool bBlueprintableComponents;
diff --git a/Engine/Source/Editor/UnrealEd/Private/AssetDeleteModel.cpp b/Engine/Source/Editor/UnrealEd/Private/AssetDeleteModel.cpp
--- a/Engine/Source/Editor/UnrealEd/Private/AssetDeleteModel.cpp
+++ b/Engine/Source/Editor/UnrealEd/Private/AssetDeleteModel.cpp
@@ -39,14 +39,25 @@ FAssetDeleteModel::FAssetDeleteModel( const TArray<UObject*>& InObjectsToDelete
 	, PendingDeleteIndex(0)
 	, ObjectsDeleted(0)
 {
+	// Take a weak copy in-case GC purges any objects from this array.
+	TArray<TWeakObjectPtr<UObject>> WeakObjectsToDelete;
+	WeakObjectsToDelete.Reserve( InObjectsToDelete.Num() );
+	for ( UObject* ObjectToDelete : InObjectsToDelete )
+	{
+		WeakObjectsToDelete.Add( ObjectToDelete );
+	}
+
 	// Purge unclaimed objects.
 	CollectGarbage( GARBAGE_COLLECTION_KEEPFLAGS );
 
-	// Create a pending delete object for each object needing to be deleted so that we can store
-	// reference tracing information.
-	for ( UObject* ObjectToDelete : InObjectsToDelete )
+	// Create a pending delete object for each object needing to be deleted so that we can store reference tracing information.
+	for ( const TWeakObjectPtr<UObject>& WeakObjectToDelete : WeakObjectsToDelete )
 	{
-		AddObjectToDelete( ObjectToDelete );
+		UObject* ObjectToDelete = WeakObjectToDelete.Get();
+		if ( ObjectToDelete )
+		{
+			AddObjectToDelete( ObjectToDelete );
+		}
 	}
 }
 
@@ -156,26 +167,29 @@ void FAssetDeleteModel::DiscoverSourceFileReferences(FPendingDelete& PendingDele
 	TArray<FString> SourceContentFiles;
 	Utils::ExtractSourceFilePaths(PendingDelete.GetObject(), SourceContentFiles);
 
-	auto MonitoredDirectories = GUnrealEd->AutoReimportManager->GetMonitoredDirectories();
+	if ( GUnrealEd->AutoReimportManager )
+	{
+		TArray<FPathAndMountPoint> MonitoredDirectories = GUnrealEd->AutoReimportManager->GetMonitoredDirectories();
 
-	// Remove anything that's not under a monitored, mounted path, or doesn't exist
-	SourceContentFiles.RemoveAll([&](const FString& InFilename){
-		for (const auto& Dir : MonitoredDirectories)
-		{
-			if (!Dir.MountPoint.IsEmpty() && InFilename.StartsWith(Dir.Path))
+		// Remove anything that's not under a monitored, mounted path, or doesn't exist
+		SourceContentFiles.RemoveAll([&] (const FString& InFilename) {
+			for ( const FPathAndMountPoint& Dir : MonitoredDirectories )
 			{
-				if (FPaths::FileExists(InFilename))
-				{
-					return false;
-				}
-				else
+				if ( !Dir.MountPoint.IsEmpty() && InFilename.StartsWith(Dir.Path) )
 				{
-					return true;
+					if ( FPaths::FileExists(InFilename) )
+					{
+						return false;
+					}
+					else
+					{
+						return true;
+					}
 				}
 			}
-		}
-		return true;
-	});
+			return true;
+		});
+	}
 
 	// Now accumulate references to the same source content file. We only offer to delete a file if it is only referenced by the deleted object(s)
 	IAssetRegistry& AssetRegistry = FModuleManager::LoadModuleChecked<FAssetRegistryModule>(TEXT("AssetRegistry")).Get();
@@ -530,11 +544,11 @@ int32 FAssetDeleteModel::GetDeletedObjectCount() const
 
 void FAssetDeleteModel::PrepareToDelete(UObject* InObject)
 {
-	if ( InObject->IsA<UObjectRedirector>() )
+	if ( UObjectRedirector* RedirectorObj = Cast<UObjectRedirector>(InObject) )
 	{
 		// Add all redirectors found in this package to the redirectors to delete list.
 		// All redirectors in this package should be fixed up.
-		UPackage* RedirectorPackage = InObject->GetOutermost();
+		UPackage* RedirectorPackage = RedirectorObj->GetOutermost();
 		TArray<UObject*> AssetsInRedirectorPackage;
 		
 		GetObjectsWithOuter(RedirectorPackage, AssetsInRedirectorPackage, /*bIncludeNestedObjects=*/false);
@@ -589,10 +603,8 @@ FPendingDelete::FPendingDelete(UObject* InObject)
 	// Blueprints actually contain 3 assets, the UBlueprint, GeneratedClass and SkeletonGeneratedClass
 	// we need to add them all to the pending deleted objects so that the memory reference detection system
 	// can know those references don't count towards the in memory references.
-	if ( InObject->IsA(UBlueprint::StaticClass()) )
+	if ( UBlueprint* BlueprintObj = Cast<UBlueprint>(InObject) )
 	{
-		UBlueprint* BlueprintObj = Cast<UBlueprint>(InObject);
-
 		if ( BlueprintObj->GeneratedClass != NULL )
 		{
 			InternalObjects.Add(BlueprintObj->GeneratedClass);
diff --git a/Engine/Source/Editor/UnrealEd/Private/AssetViewerSettings.cpp b/Engine/Source/Editor/UnrealEd/Private/AssetViewerSettings.cpp
--- a/Engine/Source/Editor/UnrealEd/Private/AssetViewerSettings.cpp
+++ b/Engine/Source/Editor/UnrealEd/Private/AssetViewerSettings.cpp
@@ -13,6 +13,14 @@ UAssetViewerSettings::UAssetViewerSettings()
 	NumProfiles = 1;
 }
 
+UAssetViewerSettings::~UAssetViewerSettings()
+{
+	if (GEditor)
+	{
+		GEditor->UnregisterForUndo(this);
+	}
+}
+
 UAssetViewerSettings* UAssetViewerSettings::Get()
 {
 	// This is a singleton, use default object
@@ -22,12 +30,19 @@ UAssetViewerSettings* UAssetViewerSettings::Get()
 	static bool bInitialized = false;
 	if (!bInitialized)
 	{
+		DefaultSettings->SetFlags(RF_Transactional);
+
 		for (FPreviewSceneProfile& Profile : DefaultSettings->Profiles)
 		{
 			Profile.LoadEnvironmentMap();
 		}
 
 		bInitialized = true;
+
+		if (GEditor)
+		{
+			GEditor->RegisterForUndo(DefaultSettings);
+		}
 	}
 
 	return DefaultSettings;
@@ -75,3 +90,16 @@ void UAssetViewerSettings::PostInitProperties()
 		ProjectSettings->AssetViewerProfileIndex = 0;
 	}
 }
+
+void UAssetViewerSettings::PostUndo(bool bSuccess)
+{
+	if (bSuccess)
+	{
+		OnAssetViewerSettingsPostUndoEvent.Broadcast();
+	}
+}
+
+void UAssetViewerSettings::PostRedo(bool bSuccess)
+{
+	PostUndo(bSuccess);
+}
