Copying //UE4/Dev-AnimPhys to //UE4/Dev-Main (Source: //UE4/Dev-AnimPhys @ 3624379)

diff --git a/Engine/Source/Editor/SkeletonEditor/Private/SkeletonTreeVirtualBoneItem.h b/Engine/Source/Editor/SkeletonEditor/Private/SkeletonTreeVirtualBoneItem.h
--- a/Engine/Source/Editor/SkeletonEditor/Private/SkeletonTreeVirtualBoneItem.h
+++ b/Engine/Source/Editor/SkeletonEditor/Private/SkeletonTreeVirtualBoneItem.h
@@ -12,6 +12,7 @@
 #include "ISkeletonTreeItem.h"
 #include "SkeletonTreeItem.h"
 #include "IEditableSkeleton.h"
+#include "BoneProxy.h"
 
 class SInlineEditableTextBlock;
 
@@ -20,23 +21,24 @@ class FSkeletonTreeVirtualBoneItem : public FSkeletonTreeItem
 public:
 	SKELETON_TREE_ITEM_TYPE(FSkeletonTreeVirtualBoneItem, FSkeletonTreeItem)
 
-	FSkeletonTreeVirtualBoneItem(const FName& InBoneName, const TSharedRef<class ISkeletonTree>& InSkeletonTree)
-		: FSkeletonTreeItem(InSkeletonTree)
-		, BoneName(InBoneName)
-	{}
+	FSkeletonTreeVirtualBoneItem(const FName& InBoneName, const TSharedRef<class ISkeletonTree>& InSkeletonTree);
+	~FSkeletonTreeVirtualBoneItem();
 
 	/** Builds the table row widget to display this info */
-	virtual TSharedRef<ITableRow> MakeTreeRowWidget(const TSharedRef<STableViewBase>& InOwnerTable, const TAttribute<FText>& InFilterText) override;
 	virtual void GenerateWidgetForNameColumn(TSharedPtr< SHorizontalBox > Box, const TAttribute<FText>& FilterText, FIsSelected InIsSelected) override;
 	virtual TSharedRef< SWidget > GenerateWidgetForDataColumn(const FName& DataColumnName) override;
 	virtual FName GetRowItemName() const override { return BoneName; }
 
 	virtual void RequestRename() override;
 	virtual void OnItemDoubleClicked() override;
+	virtual UObject* GetObject() const override { return BoneProxy; }
 
 	/** Return socket name as FText for display in skeleton tree */
 	FText GetVirtualBoneNameAsText() const { return FText::FromName(BoneName); }
 
+	/** Enable and disable the bone proxy ticking */
+	void EnableBoneProxyTick(bool bEnable);
+
 private:
 	/** Called when we are about to rename a virtual bone */
 	void OnVirtualBoneNameEditing(const FText& OriginalText);
@@ -75,4 +77,7 @@ private:
 	/** Delegate for when the context menu requests a rename */
 	DECLARE_DELEGATE(FOnRenameRequested);
 	FOnRenameRequested OnRenameRequested;
+
+	/** Bone proxy used for debug display */
+	UBoneProxy* BoneProxy;
 };
diff --git a/Engine/Source/Editor/SkeletonEditor/Public/IEditableSkeleton.h b/Engine/Source/Editor/SkeletonEditor/Public/IEditableSkeleton.h
--- a/Engine/Source/Editor/SkeletonEditor/Public/IEditableSkeleton.h
+++ b/Engine/Source/Editor/SkeletonEditor/Public/IEditableSkeleton.h
@@ -10,7 +10,7 @@ class UBlendProfile;
 class USkeletalMesh;
 
 /** Delegate fired when a set of smart names is removed */
-DECLARE_MULTICAST_DELEGATE_TwoParams(FOnSmartNameRemoved, const FName& /*InContainerName*/, const TArray<SmartName::UID_Type>& /*InNameUids*/);
+DECLARE_MULTICAST_DELEGATE_OneParam(FOnSmartNameChanged, const FName& /*InContainerName*/);
 
 /** Enum which tells us whether the parent of a socket is the skeleton or skeletal mesh */
 enum class ESocketParentType : int32
@@ -59,17 +59,14 @@ public:
 	/** Rename the specified smart name */
 	virtual void RenameSmartname(const FName& InContainerName, SmartName::UID_Type InNameUid, const FName& InNewName) = 0;
 
-	/** Remove the specified smart name */
-	virtual void RemoveSmartname(const FName& InContainerName, SmartName::UID_Type InNameUid) = 0;
-
 	/** Remove all the specified smart names and fixup animations that use them */
-	virtual void RemoveSmartnamesAndFixupAnimations(const FName& InContainerName, const TArray<SmartName::UID_Type>& InNameUids) = 0;
+	virtual void RemoveSmartnamesAndFixupAnimations(const FName& InContainerName, const TArray<FName>& InNames) = 0;
 
 	/** Sets Material Meta Data for the curve */
 	virtual void SetCurveMetaDataMaterial(const FSmartName& CurveName, bool bOverrideMaterial) = 0;
 
 	/** Sets Bone Links per curve */
-	virtual void SetCurveMetaBoneLinks(const FSmartName& CurveName, TArray<FBoneReference>& BoneLinks) = 0;
+	virtual void SetCurveMetaBoneLinks(const FSmartName& CurveName, TArray<FBoneReference>& BoneLinks, uint8 InMaxLOD) = 0;
 
 	/**
 	 * Makes sure all attached objects are valid and removes any that aren't.
@@ -95,6 +92,11 @@ public:
 	 */	
 	virtual int32 DeleteAnimNotifies(const TArray<FName>& InotifyNames) = 0;
 
+	/**
+	* Add a notify
+	*/
+	virtual void AddNotify(FName NewName) = 0;
+
 	/** 
 	 * Rename a notify
 	 * @return the number of animations modified
@@ -159,10 +161,10 @@ public:
 	virtual void RenameSlotName(const FName& InOldSlotName, const FName& InNewSlotName) = 0;
 
 	/** Register a delegate to be called when a set of smart names are removed */
-	virtual FDelegateHandle RegisterOnSmartNameRemoved(const FOnSmartNameRemoved::FDelegate& InOnSmartNameRemoved) = 0;
+	virtual FDelegateHandle RegisterOnSmartNameChanged(const FOnSmartNameChanged::FDelegate& InOnSmartNameChanged) = 0;
 
 	/** Register a delegate to be called when a set of smart names are removed */
-	virtual void UnregisterOnSmartNameRemoved(FDelegateHandle InHandle) = 0;
+	virtual void UnregisterOnSmartNameChanged(FDelegateHandle InHandle) = 0;
 
 	/** Register a delegate to be called when this skeletons notifies are changed */
 	virtual void RegisterOnNotifiesChanged(const FSimpleMulticastDelegate::FDelegate& InDelegate) = 0;
diff --git a/Engine/Source/Editor/SkeletonEditor/Public/ISkeletonEditorModule.h b/Engine/Source/Editor/SkeletonEditor/Public/ISkeletonEditorModule.h
--- a/Engine/Source/Editor/SkeletonEditor/Public/ISkeletonEditorModule.h
+++ b/Engine/Source/Editor/SkeletonEditor/Public/ISkeletonEditorModule.h
@@ -26,6 +26,9 @@ public:
 	/** Creates a new skeleton tree instance */
 	virtual TSharedRef<ISkeletonTree> CreateSkeletonTree(USkeleton* InSkeleton, const FSkeletonTreeArgs& InSkeletonTreeArgs) = 0;
 
+	/** Creates a new skeleton tree instance */
+	virtual TSharedRef<ISkeletonTree> CreateSkeletonTree(const TSharedRef<IEditableSkeleton>& InEditableSkeleton, const FSkeletonTreeArgs& InSkeletonTreeArgs) = 0;
+
 	/** Creates a new skeleton tree instance & registers a tab factory with the supplied tab factories */
 	virtual TSharedRef<class FWorkflowTabFactory> CreateSkeletonTreeTabFactory(const TSharedRef<class FWorkflowCentricApplication>& InHostingApp, const TSharedRef<ISkeletonTree>& InSkeletonTree) = 0;
 
diff --git a/Engine/Source/Editor/SkeletonEditor/Public/ISkeletonTree.h b/Engine/Source/Editor/SkeletonEditor/Public/ISkeletonTree.h
--- a/Engine/Source/Editor/SkeletonEditor/Public/ISkeletonTree.h
+++ b/Engine/Source/Editor/SkeletonEditor/Public/ISkeletonTree.h
@@ -4,33 +4,89 @@
 
 #include "CoreMinimal.h"
 #include "Widgets/SCompoundWidget.h"
+#include "ArrayView.h"
 
 struct FAssetData;
 class IPersonaPreviewScene;
 class UBlendProfile;
 struct FSelectedSocketInfo;
+class ISkeletonTreeBuilder;
+class ISkeletonTreeItem;
+class FMenuBuilder;
+class FExtender;
 
-// Called when a bone is selected
+// Called when an item is selected/deselected
+DECLARE_MULTICAST_DELEGATE_TwoParams(FOnSkeletonTreeSelectionChangedMulticast, const TArrayView<TSharedPtr<ISkeletonTreeItem>>& /* InSelectedItems */, ESelectInfo::Type /* SelectInfo */);
+
+// Called when an item is selected/deselected
+typedef FOnSkeletonTreeSelectionChangedMulticast::FDelegate FOnSkeletonTreeSelectionChanged;
+
+// Called when a bone is selected - DEPRECATED, please use FOnSkeletonTreeSelectionChangedMulticast
 DECLARE_MULTICAST_DELEGATE_OneParam(FOnObjectSelectedMulticast, UObject* /* InObject */);
 
-// Called when an object is selected
+// Called when an object is selected - DEPRECATED, please use FOnSkeletonTreeSelectionChanged
 typedef FOnObjectSelectedMulticast::FDelegate FOnObjectSelected;
 
+/** Delegate that allows custom filtering text to be shown on the filter button */
+DECLARE_DELEGATE_OneParam(FOnGetFilterText, TArray<FText>& /*InOutTextItems*/);
+
+enum class ESkeletonTreeMode
+{
+	/** Skeleton tree allows editing */
+	Editor,
+
+	/** Skeleton tree allows picking of tree elements */
+	Picker,
+};
+
 /** Init params for a skeleton tree widget */
 struct FSkeletonTreeArgs
 {
-	FSkeletonTreeArgs(FSimpleMulticastDelegate& InOnPostUndo)
-		: OnPostUndo(InOnPostUndo)
+	FSkeletonTreeArgs()
+		: Mode(ESkeletonTreeMode::Editor)
+		, bShowBlendProfiles(true)
+		, bShowFilterMenu(true)
+		, bAllowMeshOperations(true)
+		, bAllowSkeletonOperations(true)
 	{}
 
-	/** Multicast delegate the tree will register with for undo messages */
-	FSimpleMulticastDelegate& OnPostUndo;
+	PRAGMA_DISABLE_DEPRECATION_WARNINGS
+	~FSkeletonTreeArgs()
+	{}
+	PRAGMA_ENABLE_DEPRECATION_WARNINGS
 
 	/** Delegate called by the tree when a socket is selected */
+	FOnSkeletonTreeSelectionChanged OnSelectionChanged;
+
+	DEPRECATED(4.17, "Please use OnSelectionChanged")
 	FOnObjectSelected OnObjectSelected;
 
+	/** Delegate that allows custom filtering text to be shown on the filter button */
+	FOnGetFilterText OnGetFilterText;
+
 	/** Optional preview scene that we can pair with */
-	TSharedPtr<class IPersonaPreviewScene> PreviewScene;
+	TSharedPtr<IPersonaPreviewScene> PreviewScene;
+
+	/** Optional builder to allow for custom tree construction */
+	TSharedPtr<ISkeletonTreeBuilder> Builder;
+
+	/** Menu extenders applied to context and filter menus */
+	TSharedPtr<FExtender> Extenders;
+
+	/** The mode that this skeleton tree is in */
+	ESkeletonTreeMode Mode;
+
+	/** Whether to show the blend profiles editor for the skeleton being displayed */
+	bool bShowBlendProfiles;
+
+	/** Whether to show the filter menu to allow filtering of active bones, sockets etc. */
+	bool bShowFilterMenu;
+
+	/** Whether to allow operations that modify the mesh */
+	bool bAllowMeshOperations;
+
+	/** Whether to allow operations that modify the skeleton */
+	bool bAllowSkeletonOperations;
 };
 
 /** Interface used to deal with skeleton editing UI */
@@ -38,11 +94,17 @@ class SKELETONEDITOR_API ISkeletonTree : public SCompoundWidget
 {
 public:
 	struct Columns
-{
+	{
 		static const FName Name;
 		static const FName Retargeting;
 		static const FName BlendProfile;
-};
+	};
+
+	/** Manually refresh the tree */
+	virtual void Refresh() = 0;
+
+	/** Manually refresh the tree filter */
+	virtual void RefreshFilter() = 0;
 
 	/** Get editable skeleton that this widget is editing */
 	virtual TSharedRef<class IEditableSkeleton> GetEditableSkeleton() const = 0;
@@ -62,18 +124,33 @@ public:
 	/** Deselect everything that is currently selected */
 	virtual void DeselectAll() = 0;
 
+	/** Get the selected items */
+	virtual TArray<TSharedPtr<ISkeletonTreeItem>> GetSelectedItems() const = 0;
+
+	/** Select items using the passed in predicate  */
+	virtual void SelectItemsBy(TFunctionRef<bool(const TSharedRef<ISkeletonTreeItem>& /* InItem */, bool& /* bInOutExpand */)> Predicate) const = 0;
+
 	/** Duplicate the socket and select it */
 	virtual void DuplicateAndSelectSocket(const FSelectedSocketInfo& SocketInfoToDuplicate, const FName& NewParentBoneName = FName()) = 0;
 
-	/** Registers a delegate to be called when the selected object is changed */
-	virtual void RegisterOnObjectSelected(const FOnObjectSelected& Delegate) = 0;
+	/** Registers a delegate to be called when the selected items have changed */
+	virtual FDelegateHandle RegisterOnSelectionChanged(const FOnSkeletonTreeSelectionChanged& Delegate) = 0;
 
-	/** Unregisters a delegate to be called when the selected object is changed */
-	virtual void UnregisterOnObjectSelected(SWidget* Widget) = 0;
+	/** Unregisters a delegate to be called when the selected items have changed */
+	virtual void UnregisterOnSelectionChanged(FDelegateHandle DelegateHandle) = 0;
 
 	/** Gets the currently selected blend profile */
 	virtual UBlendProfile* GetSelectedBlendProfile() = 0;
 
 	/** Attached the supplied assets to the tree to the specified attach item (bone/socket) */
 	virtual void AttachAssets(const TSharedRef<class ISkeletonTreeItem>& TargetItem, const TArray<FAssetData>& AssetData) = 0;
+
+	/** Get the search box widget, if any, for this tree */
+	virtual TSharedPtr<SWidget> GetSearchWidget() const = 0;
+
+	DEPRECATED(4.17, "Please use RegisterOnSelectionChanged")
+	virtual void RegisterOnObjectSelected(const FOnObjectSelected& Delegate) = 0;
+
+	DEPRECATED(4.17, "Please use UnregisterOnSelectionChanged")
+	virtual void UnregisterOnObjectSelected(SWidget* Widget) = 0;
 };
diff --git a/Engine/Source/Editor/SkeletonEditor/Public/ISkeletonTreeBuilder.h b/Engine/Source/Editor/SkeletonEditor/Public/ISkeletonTreeBuilder.h
--- a/Engine/Source/Editor/SkeletonEditor/Public/ISkeletonTreeBuilder.h
+++ b/Engine/Source/Editor/SkeletonEditor/Public/ISkeletonTreeBuilder.h
@@ -4,8 +4,11 @@
 
 #include "CoreMinimal.h"
 #include "Containers/ArrayView.h"
+#include "TextFilterExpressionEvaluator.h"
 
 class ISkeletonTreeItem;
+enum class ESkeletonTreeFilterResult : int32;
+class FTextFilterExpressionEvaluator;
 
 /** Output struct for builders to use */
 struct SKELETONEDITOR_API FSkeletonTreeBuilderOutput
@@ -20,16 +23,18 @@ struct SKELETONEDITOR_API FSkeletonTreeBuilderOutput
 	 * @param	InItem			The item to add
 	 * @param	InParentName	The name of the item's parent
 	 * @param	InParentTypes	The types of items to search. If this is empty all items will be searched.
+	 * @param	bAddToHead		Whether to add the item to the start or end of the parent's children array
 	 */
-	void Add(const TSharedPtr<class ISkeletonTreeItem>& InItem, const FName& InParentName, TArrayView<const FName> InParentTypes);
+	void Add(const TSharedPtr<class ISkeletonTreeItem>& InItem, const FName& InParentName, TArrayView<const FName> InParentTypes, bool bAddToHead = false);
 
 	/** 
 	 * Add an item to the output
 	 * @param	InItem			The item to add
 	 * @param	InParentName	The name of the item's parent
 	 * @param	InParentType	The type of items to search. If this is NAME_None all items will be searched.
+	 * @param	bAddToHead		Whether to add the item to the start or end of the parent's children array
 	 */
-	void Add(const TSharedPtr<class ISkeletonTreeItem>& InItem, const FName& InParentName, const FName& InParentType);
+	void Add(const TSharedPtr<class ISkeletonTreeItem>& InItem, const FName& InParentName, const FName& InParentType, bool bAddToHead = false);
 
 	/** 
 	 * Find the item with the specified name
@@ -55,21 +60,47 @@ private:
 	TArray<TSharedPtr<class ISkeletonTreeItem>>& LinearItems;
 };
 
+/** Filter utility class */
+class FSkeletonTreeFilterContext : public ITextFilterExpressionContext
+{
+public:
+	explicit FSkeletonTreeFilterContext(const FName& InName)
+		: Name(InName)
+	{
+	}
+
+	virtual bool TestBasicStringExpression(const FTextFilterString& InValue, const ETextFilterTextComparisonMode InTextComparisonMode) const override
+	{
+		return TextFilterUtils::TestBasicStringExpression(Name.ToString(), InValue, InTextComparisonMode);
+	}
+
+	virtual bool TestComplexExpression(const FName& InKey, const FTextFilterString& InValue, const ETextFilterComparisonOperation InComparisonOperation, const ETextFilterTextComparisonMode InTextComparisonMode) const override
+	{
+		return false;
+	}
+
+private:
+	FName Name;
+};
+
 /** Basic filter used when re-filtering the tree */
 struct FSkeletonTreeFilterArgs
 {
-	FSkeletonTreeFilterArgs()
-		: bWillFilter(false)
+	FSkeletonTreeFilterArgs(TSharedPtr<FTextFilterExpressionEvaluator> InTextFilter)
+		: TextFilter(InTextFilter)
 		, bFlattenHierarchyOnFilter(true)
 	{}
 
-	/** Whether or not we have any kind of filtering */
-	bool bWillFilter;
+	/** The text filter we are using, if any */
+	TSharedPtr<FTextFilterExpressionEvaluator> TextFilter;
 
 	/** Whether to flatten the hierarchy so filtered items appear in a linear list */
 	bool bFlattenHierarchyOnFilter;
 };
 
+/** Delegate used to filter an item. */
+DECLARE_DELEGATE_RetVal_TwoParams(ESkeletonTreeFilterResult, FOnFilterSkeletonTreeItem, const FSkeletonTreeFilterArgs& /*InArgs*/, const TSharedPtr<class ISkeletonTreeItem>& /*InItem*/);
+
 /** Interface to implement to provide custom build logic to skeleton trees */
 class ISkeletonTreeBuilder
 {
@@ -77,13 +108,27 @@ public:
 
 	virtual ~ISkeletonTreeBuilder() {};
 
+	/** Setup this builder with links to the tree and preview scene */
+	virtual void Initialize(const TSharedRef<class ISkeletonTree>& InSkeletonTree, const TSharedPtr<class IPersonaPreviewScene>& InPreviewScene, FOnFilterSkeletonTreeItem InOnFilterSkeletonTreeItem) = 0;
+
 	/**
 	 * Build an array of skeleton tree items to display in the tree.
-	 * @param	InArgs			Parameters to use when building the tree
 	 * @param	Output			The items that are built by this builder
 	 */
 	virtual void Build(FSkeletonTreeBuilderOutput& Output) = 0;
 
 	/** Apply filtering to the tree items */
 	virtual void Filter(const FSkeletonTreeFilterArgs& InArgs, const TArray<TSharedPtr<class ISkeletonTreeItem>>& InItems, TArray<TSharedPtr<class ISkeletonTreeItem>>& OutFilteredItems) = 0;
+
+	/** Allows the builder to contribute to filtering an item */
+	virtual ESkeletonTreeFilterResult FilterItem(const FSkeletonTreeFilterArgs& InArgs, const TSharedPtr<class ISkeletonTreeItem>& InItem) = 0;
+
+	/** Get whether this builder is showing bones */
+	virtual bool IsShowingBones() const = 0;
+
+	/** Get whether this builder is showing bones */
+	virtual bool IsShowingSockets() const = 0;
+
+	/** Get whether this builder is showing bones */
+	virtual bool IsShowingAttachedAssets() const = 0;
 };
diff --git a/Engine/Source/Editor/SkeletonEditor/Public/ISkeletonTreeItem.h b/Engine/Source/Editor/SkeletonEditor/Public/ISkeletonTreeItem.h
--- a/Engine/Source/Editor/SkeletonEditor/Public/ISkeletonTreeItem.h
+++ b/Engine/Source/Editor/SkeletonEditor/Public/ISkeletonTreeItem.h
@@ -30,8 +30,8 @@ enum class ESkeletonTreeFilterResult : int32
 
 #define SKELETON_TREE_BASE_ITEM_TYPE(TYPE) \
 	static const FName& GetTypeId() { static FName Type(TEXT(#TYPE)); return Type; } \
-	virtual bool IsOfTypeByName(const FName& Type) const { return GetTypeId() == Type; } \
-	virtual FName GetTypeName() const { return GetTypeId(); }
+	virtual bool IsOfTypeByName(const FName& Type) const { return TYPE::GetTypeId() == Type; } \
+	virtual FName GetTypeName() const { return TYPE::GetTypeId(); }
 
 /** Interface for a skeleton tree item */
 class ISkeletonTreeItem : public TSharedFromThis<ISkeletonTreeItem>
@@ -40,7 +40,7 @@ public:
 	SKELETON_TREE_BASE_ITEM_TYPE(ISkeletonTreeItem)
 
 	/** Check if this item can cast safely to the specified template type */
-	template<class TType> bool IsOfType() const
+	template<typename TType> bool IsOfType() const
 	{
 		return IsOfTypeByName(TType::GetTypeId());
 	}
@@ -89,6 +89,12 @@ public:
 	/** Handle a drag and drop drop event */
 	virtual FReply HandleDrop(const FDragDropEvent& DragDropEvent) = 0;
 
+	/** Get this item's parent  */
+	virtual TSharedPtr<ISkeletonTreeItem> GetParent() const = 0;
+
+	/** Set this item's parent */
+	virtual void SetParent(TSharedPtr<ISkeletonTreeItem> InParent) = 0;
+
 	/** The array of children for this item */
 	virtual TArray<TSharedPtr<ISkeletonTreeItem>>& GetChildren() = 0;
 
@@ -106,6 +112,12 @@ public:
 
 	/** Set the current filter result */
 	virtual void SetFilterResult(ESkeletonTreeFilterResult InResult) = 0;
+
+	/** Get the object represented by this item, if any */
+	virtual UObject* GetObject() const = 0;
+
+	/** Get whether this item begins expanded or not */
+	virtual bool IsInitiallyExpanded() const = 0;
 };
 
 /**
@@ -120,4 +132,5 @@ public:
  */
 #define SKELETON_TREE_ITEM_TYPE(TYPE, BASE) \
 	static const FName& GetTypeId() { static FName Type(TEXT(#TYPE)); return Type; } \
-	virtual bool IsOfTypeByName(const FName& Type) const override { return GetTypeId() == Type || BASE::IsOfTypeByName(Type); }
+	virtual bool IsOfTypeByName(const FName& Type) const override { return GetTypeId() == Type || BASE::IsOfTypeByName(Type); } \
+	virtual FName GetTypeName() const { return TYPE::GetTypeId(); }
diff --git a/Engine/Source/Editor/SkeletonEditor/Public/SkeletonTreeBuilder.h b/Engine/Source/Editor/SkeletonEditor/Public/SkeletonTreeBuilder.h
--- a/Engine/Source/Editor/SkeletonEditor/Public/SkeletonTreeBuilder.h
+++ b/Engine/Source/Editor/SkeletonEditor/Public/SkeletonTreeBuilder.h
@@ -14,12 +14,16 @@ class USkeletalMeshSocket;
 /** Options for skeleton building */
 struct FSkeletonTreeBuilderArgs
 {
-	FSkeletonTreeBuilderArgs()
-		: bShowSockets(true)
-		, bShowAttachedAssets(true)
-		, bShowVirtualBones(true)
+	FSkeletonTreeBuilderArgs(bool bInShowBones = true, bool bInShowSockets = true, bool bInShowAttachedAssets = true, bool bInShowVirtualBones = true)
+		: bShowBones(bInShowBones)
+		, bShowSockets(bInShowSockets)
+		, bShowAttachedAssets(bInShowAttachedAssets)
+		, bShowVirtualBones(bInShowVirtualBones)
 	{}
 
+	/** Whether we should show bones */
+	bool bShowBones;
+
 	/** Whether we should show sockets */
 	bool bShowSockets;
 
@@ -52,19 +56,21 @@ enum class ESocketFilter : int32
 	Count
 };
 
-/** Delegate used to filter an item. */
-DECLARE_DELEGATE_RetVal_OneParam(ESkeletonTreeFilterResult, FOnFilterSkeletonTreeItem, const TSharedPtr<class ISkeletonTreeItem>& /*InItem*/);
-
 class SKELETONEDITOR_API FSkeletonTreeBuilder : public ISkeletonTreeBuilder
 {
 public:
-	FSkeletonTreeBuilder(const FSkeletonTreeBuilderArgs& InBuilderArgs, FOnFilterSkeletonTreeItem InOnFilterSkeletonTreeItem, const TSharedRef<class ISkeletonTree>& InSkeletonTree, const TSharedPtr<class IPersonaPreviewScene>& InPreviewScene);
+	FSkeletonTreeBuilder(const FSkeletonTreeBuilderArgs& InBuilderArgs);
 
 	/** ISkeletonTreeBuilder interface */
+	virtual void Initialize(const TSharedRef<class ISkeletonTree>& InSkeletonTree, const TSharedPtr<class IPersonaPreviewScene>& InPreviewScene, FOnFilterSkeletonTreeItem InOnFilterSkeletonTreeItem) override;
 	virtual void Build(FSkeletonTreeBuilderOutput& Output) override;
 	virtual void Filter(const FSkeletonTreeFilterArgs& InArgs, const TArray<TSharedPtr<class ISkeletonTreeItem>>& InItems, TArray<TSharedPtr<class ISkeletonTreeItem>>& OutFilteredItems) override;
+	virtual ESkeletonTreeFilterResult FilterItem(const FSkeletonTreeFilterArgs& InArgs, const TSharedPtr<class ISkeletonTreeItem>& InItem) override;
+	virtual bool IsShowingBones() const override;
+	virtual bool IsShowingSockets() const override;
+	virtual bool IsShowingAttachedAssets() const override;
 
-private:
+protected:
 	/** Add Bones */
 	void AddBones(FSkeletonTreeBuilderOutput& Output);
 
@@ -94,7 +100,7 @@ private:
 	/** Helper function for filtering */
 	ESkeletonTreeFilterResult FilterRecursive(const FSkeletonTreeFilterArgs& InArgs, const TSharedPtr<ISkeletonTreeItem>& InItem, TArray<TSharedPtr<ISkeletonTreeItem>>& OutFilteredItems);
 
-private:
+protected:
 	/** Options for building */
 	FSkeletonTreeBuilderArgs BuilderArgs;
 
diff --git a/Engine/Source/Editor/SkeletonEditor/Public/SkeletonTreeItem.h b/Engine/Source/Editor/SkeletonEditor/Public/SkeletonTreeItem.h
--- a/Engine/Source/Editor/SkeletonEditor/Public/SkeletonTreeItem.h
+++ b/Engine/Source/Editor/SkeletonEditor/Public/SkeletonTreeItem.h
@@ -15,6 +15,8 @@
 class SKELETONEDITOR_API FSkeletonTreeItem : public ISkeletonTreeItem
 {
 public:
+	SKELETON_TREE_ITEM_TYPE(FSkeletonTreeItem, ISkeletonTreeItem)
+
 	friend class FSkeletonTreeBuilder;
 
 	FSkeletonTreeItem(const TSharedRef<class ISkeletonTree>& InSkeletonTree)
@@ -35,14 +37,22 @@ public:
 	virtual void HandleDragEnter(const FDragDropEvent& DragDropEvent) override {}
 	virtual void HandleDragLeave(const FDragDropEvent& DragDropEvent) override {}
 	virtual FReply HandleDrop(const FDragDropEvent& DragDropEvent) override { return FReply::Unhandled(); }
+	virtual TSharedPtr<ISkeletonTreeItem> GetParent() const override { return Parent.Pin(); }
+	virtual void SetParent(TSharedPtr<ISkeletonTreeItem> InParent) override { Parent = InParent; }
 	virtual TArray<TSharedPtr<ISkeletonTreeItem>>& GetChildren() override { return Children; }
 	virtual TArray<TSharedPtr<ISkeletonTreeItem>>& GetFilteredChildren() override { return FilteredChildren; }
 	virtual TSharedRef<class ISkeletonTree> GetSkeletonTree() const override { return SkeletonTreePtr.Pin().ToSharedRef(); }
 	virtual TSharedRef<class IEditableSkeleton> GetEditableSkeleton() const override { return GetSkeletonTree()->GetEditableSkeleton(); }
 	virtual ESkeletonTreeFilterResult GetFilterResult() const override { return FilterResult; }
 	virtual void SetFilterResult(ESkeletonTreeFilterResult InResult) override { FilterResult = InResult; }
+	virtual FReply OnDragDetected(const FGeometry& MyGeometry, const FPointerEvent& MouseEvent) { return FReply::Unhandled(); }
+	virtual UObject* GetObject() const override { return nullptr; }
+	virtual bool IsInitiallyExpanded() const override { return true; }
 
 protected:
+	/** The parent of this item */
+	TWeakPtr<ISkeletonTreeItem> Parent;
+
 	/** The children of this item */
 	TArray<TSharedPtr<ISkeletonTreeItem>> Children;
 
diff --git a/Engine/Source/Editor/StaticMeshEditor/Private/StaticMeshEditor.cpp b/Engine/Source/Editor/StaticMeshEditor/Private/StaticMeshEditor.cpp
--- a/Engine/Source/Editor/StaticMeshEditor/Private/StaticMeshEditor.cpp
+++ b/Engine/Source/Editor/StaticMeshEditor/Private/StaticMeshEditor.cpp
@@ -656,13 +656,13 @@ bool FStaticMeshEditor::IsPrimValid(const FPrimData& InPrimData) const
 
 		switch (InPrimData.PrimType)
 		{
-		case KPT_Sphere:
+		case EAggCollisionShape::Sphere:
 			return AggGeom->SphereElems.IsValidIndex(InPrimData.PrimIndex);
-		case KPT_Box:
+		case EAggCollisionShape::Box:
 			return AggGeom->BoxElems.IsValidIndex(InPrimData.PrimIndex);
-		case KPT_Sphyl:
+		case EAggCollisionShape::Sphyl:
 			return AggGeom->SphylElems.IsValidIndex(InPrimData.PrimIndex);
-		case KPT_Convex:
+		case EAggCollisionShape::Convex:
 			return AggGeom->ConvexElems.IsValidIndex(InPrimData.PrimIndex);
 		}
 	}
@@ -740,25 +740,25 @@ void FStaticMeshEditor::DuplicateSelectedPrims(const FVector* InOffset)
 			check(IsPrimValid(PrimData));
 			switch (PrimData.PrimType)
 			{
-			case KPT_Sphere:
+			case EAggCollisionShape::Sphere:
 				{					
 					const FKSphereElem SphereElem = AggGeom->SphereElems[PrimData.PrimIndex];
 					PrimData.PrimIndex = AggGeom->SphereElems.Add(SphereElem);
 				}
 				break;
-			case KPT_Box:
+			case EAggCollisionShape::Box:
 				{
 					const FKBoxElem BoxElem = AggGeom->BoxElems[PrimData.PrimIndex];
 					PrimData.PrimIndex = AggGeom->BoxElems.Add(BoxElem);
 				}
 				break;
-			case KPT_Sphyl:
+			case EAggCollisionShape::Sphyl:
 				{
 					const FKSphylElem SphylElem = AggGeom->SphylElems[PrimData.PrimIndex];
 					PrimData.PrimIndex = AggGeom->SphylElems.Add(SphylElem);
 				}
 				break;
-			case KPT_Convex:
+			case EAggCollisionShape::Convex:
 				{
 					const FKConvexElem ConvexElem = AggGeom->ConvexElems[PrimData.PrimIndex];
 					PrimData.PrimIndex = AggGeom->ConvexElems.Add(ConvexElem);
@@ -865,16 +865,16 @@ void FStaticMeshEditor::ScaleSelectedPrims(const FVector& InScale)
 		check(IsPrimValid(PrimData));
 		switch (PrimData.PrimType)
 		{
-		case KPT_Sphere:
+		case EAggCollisionShape::Sphere:
 			AggGeom->SphereElems[PrimData.PrimIndex].ScaleElem(SimplePrimitiveScaleSpeedFactor * ModifiedScale, MinPrimSize);
 			break;
-		case KPT_Box:
+		case EAggCollisionShape::Box:
 			AggGeom->BoxElems[PrimData.PrimIndex].ScaleElem(SimplePrimitiveScaleSpeedFactor * ModifiedScale, MinPrimSize);
 			break;
-		case KPT_Sphyl:
+		case EAggCollisionShape::Sphyl:
 			AggGeom->SphylElems[PrimData.PrimIndex].ScaleElem(SimplePrimitiveScaleSpeedFactor * ModifiedScale, MinPrimSize);
 			break;
-		case KPT_Convex:
+		case EAggCollisionShape::Convex:
 			AggGeom->ConvexElems[PrimData.PrimIndex].ScaleElem(ModifiedScale, MinPrimSize);
 			break;
 		}
@@ -899,16 +899,16 @@ bool FStaticMeshEditor::CalcSelectedPrimsAABB(FBox &OutBox) const
 		check(IsPrimValid(PrimData));
 		switch (PrimData.PrimType)
 		{
-		case KPT_Sphere:
+		case EAggCollisionShape::Sphere:
 			OutBox += AggGeom->SphereElems[PrimData.PrimIndex].CalcAABB(FTransform::Identity, 1.f);
 			break;
-		case KPT_Box:
+		case EAggCollisionShape::Box:
 			OutBox += AggGeom->BoxElems[PrimData.PrimIndex].CalcAABB(FTransform::Identity, 1.f);
 			break;
-		case KPT_Sphyl:
+		case EAggCollisionShape::Sphyl:
 			OutBox += AggGeom->SphylElems[PrimData.PrimIndex].CalcAABB(FTransform::Identity, 1.f);
 			break;
-		case KPT_Convex:
+		case EAggCollisionShape::Convex:
 			OutBox += AggGeom->ConvexElems[PrimData.PrimIndex].CalcAABB(FTransform::Identity, FVector(1.f));
 			break;
 		}
@@ -929,16 +929,16 @@ bool FStaticMeshEditor::GetLastSelectedPrimTransform(FTransform& OutTransform) c
 		check(IsPrimValid(PrimData));
 		switch (PrimData.PrimType)
 		{
-		case KPT_Sphere:
+		case EAggCollisionShape::Sphere:
 			OutTransform = AggGeom->SphereElems[PrimData.PrimIndex].GetTransform();
 			break;
-		case KPT_Box:
+		case EAggCollisionShape::Box:
 			OutTransform = AggGeom->BoxElems[PrimData.PrimIndex].GetTransform();
 			break;
-		case KPT_Sphyl:
+		case EAggCollisionShape::Sphyl:
 			OutTransform = AggGeom->SphylElems[PrimData.PrimIndex].GetTransform();
 			break;
-		case KPT_Convex:
+		case EAggCollisionShape::Convex:
 			OutTransform = AggGeom->ConvexElems[PrimData.PrimIndex].GetTransform();
 			break;
 		}
@@ -955,13 +955,13 @@ FTransform FStaticMeshEditor::GetPrimTransform(const FPrimData& InPrimData) cons
 	check(IsPrimValid(InPrimData));
 	switch (InPrimData.PrimType)
 	{
-	case KPT_Sphere:
+	case EAggCollisionShape::Sphere:
 		return AggGeom->SphereElems[InPrimData.PrimIndex].GetTransform();
-	case KPT_Box:
+	case EAggCollisionShape::Box:
 		return AggGeom->BoxElems[InPrimData.PrimIndex].GetTransform();
-	case KPT_Sphyl:
+	case EAggCollisionShape::Sphyl:
 		return AggGeom->SphylElems[InPrimData.PrimIndex].GetTransform();
-	case KPT_Convex:
+	case EAggCollisionShape::Convex:
 		return AggGeom->ConvexElems[InPrimData.PrimIndex].GetTransform();
 	}
 	return FTransform::Identity;
@@ -976,16 +976,16 @@ void FStaticMeshEditor::SetPrimTransform(const FPrimData& InPrimData, const FTra
 	check(IsPrimValid(InPrimData));
 	switch (InPrimData.PrimType)
 	{
-	case KPT_Sphere:
+	case EAggCollisionShape::Sphere:
 		AggGeom->SphereElems[InPrimData.PrimIndex].SetTransform(InPrimTransform);
 		break;
-	case KPT_Box:
+	case EAggCollisionShape::Box:
 		AggGeom->BoxElems[InPrimData.PrimIndex].SetTransform(InPrimTransform);
 		break;
-	case KPT_Sphyl:
+	case EAggCollisionShape::Sphyl:
 		AggGeom->SphylElems[InPrimData.PrimIndex].SetTransform(InPrimTransform);
 		break;
-	case KPT_Convex:
+	case EAggCollisionShape::Convex:
 		AggGeom->ConvexElems[InPrimData.PrimIndex].SetTransform(InPrimTransform);
 		break;
 	}
@@ -1003,7 +1003,7 @@ bool FStaticMeshEditor::OverlapsExistingPrim(const FPrimData& InPrimData) const
 	check(IsPrimValid(InPrimData));
 	switch (InPrimData.PrimType)
 	{
-	case KPT_Sphere:
+	case EAggCollisionShape::Sphere:
 		{
 			const FKSphereElem InSphereElem = AggGeom->SphereElems[InPrimData.PrimIndex];
 			const FTransform InElemTM = InSphereElem.GetTransform();
@@ -1023,7 +1023,7 @@ bool FStaticMeshEditor::OverlapsExistingPrim(const FPrimData& InPrimData) const
 			}
 		}
 		break;
-	case KPT_Box:
+	case EAggCollisionShape::Box:
 		{
 			const FKBoxElem InBoxElem = AggGeom->BoxElems[InPrimData.PrimIndex];
 			const FTransform InElemTM = InBoxElem.GetTransform();
@@ -1043,7 +1043,7 @@ bool FStaticMeshEditor::OverlapsExistingPrim(const FPrimData& InPrimData) const
 			}
 		}
 		break;
-	case KPT_Sphyl:
+	case EAggCollisionShape::Sphyl:
 		{
 			const FKSphylElem InSphylElem = AggGeom->SphylElems[InPrimData.PrimIndex];
 			const FTransform InElemTM = InSphylElem.GetTransform();
@@ -1063,7 +1063,7 @@ bool FStaticMeshEditor::OverlapsExistingPrim(const FPrimData& InPrimData) const
 			}
 		}
 		break;
-	case KPT_Convex:
+	case EAggCollisionShape::Convex:
 		{
 			const FKConvexElem InConvexElem = AggGeom->ConvexElems[InPrimData.PrimIndex];
 			const FTransform InElemTM = InConvexElem.GetTransform();
@@ -1288,7 +1288,7 @@ void FStaticMeshEditor::GenerateKDop(const FVector* Directions, uint32 NumDirect
 		{
 			FEngineAnalytics::GetProvider().RecordEvent(TEXT("Editor.Usage.StaticMesh.Collision"), TEXT("Type"), TEXT("KDop Collision"));
 		}
-		const FPrimData PrimData = FPrimData(KPT_Convex, PrimIndex);
+		const FPrimData PrimData = FPrimData(EAggCollisionShape::Convex, PrimIndex);
 		ClearSelectedPrims();
 		AddSelectedPrim(PrimData, true);
 		// Don't 'nudge' KDop prims, as they are fitted specifically around the geometry
@@ -1308,7 +1308,7 @@ void FStaticMeshEditor::OnCollisionBox()
 		{
 			FEngineAnalytics::GetProvider().RecordEvent(TEXT("Editor.Usage.StaticMesh.Collision"), TEXT("Type"), TEXT("Box Collision"));
 		}
-		const FPrimData PrimData = FPrimData(KPT_Box, PrimIndex);
+		const FPrimData PrimData = FPrimData(EAggCollisionShape::Box, PrimIndex);
 		ClearSelectedPrims();
 		AddSelectedPrim(PrimData, true);
 		while( OverlapsExistingPrim(PrimData) )
@@ -1331,7 +1331,7 @@ void FStaticMeshEditor::OnCollisionSphere()
 		{
 			FEngineAnalytics::GetProvider().RecordEvent(TEXT("Editor.Usage.StaticMesh.Collision"), TEXT("Type"), TEXT("Sphere Collision"));
 		}
-		const FPrimData PrimData = FPrimData(KPT_Sphere, PrimIndex);
+		const FPrimData PrimData = FPrimData(EAggCollisionShape::Sphere, PrimIndex);
 		ClearSelectedPrims();
 		AddSelectedPrim(PrimData, true);
 		while( OverlapsExistingPrim(PrimData) )
@@ -1354,7 +1354,7 @@ void FStaticMeshEditor::OnCollisionSphyl()
 		{
 			FEngineAnalytics::GetProvider().RecordEvent(TEXT("Editor.Usage.StaticMesh.Collision"), TEXT("Type"), TEXT("Capsule Collision"));
 		}
-		const FPrimData PrimData = FPrimData(KPT_Sphyl, PrimIndex);
+		const FPrimData PrimData = FPrimData(EAggCollisionShape::Sphyl, PrimIndex);
 		ClearSelectedPrims();
 		AddSelectedPrim(PrimData, true);
 		while( OverlapsExistingPrim(PrimData) )
@@ -1470,7 +1470,7 @@ void FStaticMeshEditor::OnConvertBoxToConvexCollision()
 				FKAggregateGeom* AggGeom = &StaticMesh->BodySetup->AggGeom;
 				for (int32 i = 0; i < NumBoxElems; ++i)
 				{
-					AddSelectedPrim(FPrimData(KPT_Convex, (AggGeom->ConvexElems.Num() - (i+1))), false);
+					AddSelectedPrim(FPrimData(EAggCollisionShape::Convex, (AggGeom->ConvexElems.Num() - (i+1))), false);
 				}
 
 				RefreshCollisionChange(*StaticMesh);
@@ -1825,16 +1825,16 @@ void FStaticMeshEditor::DeleteSelectedPrims()
 			check(IsPrimValid(PrimData));
 			switch (PrimData.PrimType)
 			{
-			case KPT_Sphere:
+			case EAggCollisionShape::Sphere:
 				AggGeom->SphereElems.RemoveAt(PrimData.PrimIndex);
 				break;
-			case KPT_Box:
+			case EAggCollisionShape::Box:
 				AggGeom->BoxElems.RemoveAt(PrimData.PrimIndex);
 				break;
-			case KPT_Sphyl:
+			case EAggCollisionShape::Sphyl:
 				AggGeom->SphylElems.RemoveAt(PrimData.PrimIndex);
 				break;
-			case KPT_Convex:
+			case EAggCollisionShape::Convex:
 				AggGeom->ConvexElems.RemoveAt(PrimData.PrimIndex);
 				break;
 			}
diff --git a/Engine/Source/Editor/StaticMeshEditor/Private/StaticMeshEditorViewportClient.cpp b/Engine/Source/Editor/StaticMeshEditor/Private/StaticMeshEditorViewportClient.cpp
--- a/Engine/Source/Editor/StaticMeshEditor/Private/StaticMeshEditorViewportClient.cpp
+++ b/Engine/Source/Editor/StaticMeshEditor/Private/StaticMeshEditorViewportClient.cpp
@@ -126,7 +126,7 @@ struct HSMECollisionProxy : public HHitProxy
 		HHitProxy(HPP_UI),
 		PrimData(InPrimData) {}
 
-	HSMECollisionProxy(EKCollisionPrimitiveType InPrimType, int32 InPrimIndex) :
+	HSMECollisionProxy(EAggCollisionShape::Type InPrimType, int32 InPrimIndex) :
 		HHitProxy(HPP_UI),
 		PrimData(InPrimType, InPrimIndex) {}
 };
@@ -465,7 +465,7 @@ void FStaticMeshEditorViewportClient::Draw(const FSceneView* View,FPrimitiveDraw
 
 		for (int32 i = 0; i < AggGeom->SphereElems.Num(); ++i)
 		{
-			HSMECollisionProxy* HitProxy = new HSMECollisionProxy(KPT_Sphere, i);
+			HSMECollisionProxy* HitProxy = new HSMECollisionProxy(EAggCollisionShape::Sphere, i);
 			PDI->SetHitProxy(HitProxy);
 
 			const FColor CollisionColor = StaticMeshEditor->IsSelectedPrim(HitProxy->PrimData) ? SelectedColor : UnselectedColor;
@@ -478,7 +478,7 @@ void FStaticMeshEditorViewportClient::Draw(const FSceneView* View,FPrimitiveDraw
 
 		for (int32 i = 0; i < AggGeom->BoxElems.Num(); ++i)
 		{
-			HSMECollisionProxy* HitProxy = new HSMECollisionProxy(KPT_Box, i);
+			HSMECollisionProxy* HitProxy = new HSMECollisionProxy(EAggCollisionShape::Box, i);
 			PDI->SetHitProxy(HitProxy);
 
 			const FColor CollisionColor = StaticMeshEditor->IsSelectedPrim(HitProxy->PrimData) ? SelectedColor : UnselectedColor;
@@ -491,7 +491,7 @@ void FStaticMeshEditorViewportClient::Draw(const FSceneView* View,FPrimitiveDraw
 
 		for (int32 i = 0; i < AggGeom->SphylElems.Num(); ++i)
 		{
-			HSMECollisionProxy* HitProxy = new HSMECollisionProxy(KPT_Sphyl, i);
+			HSMECollisionProxy* HitProxy = new HSMECollisionProxy(EAggCollisionShape::Sphyl, i);
 			PDI->SetHitProxy(HitProxy);
 
 			const FColor CollisionColor = StaticMeshEditor->IsSelectedPrim(HitProxy->PrimData) ? SelectedColor : UnselectedColor;
@@ -504,7 +504,7 @@ void FStaticMeshEditorViewportClient::Draw(const FSceneView* View,FPrimitiveDraw
 
 		for (int32 i = 0; i < AggGeom->ConvexElems.Num(); ++i)
 		{
-			HSMECollisionProxy* HitProxy = new HSMECollisionProxy(KPT_Convex, i);
+			HSMECollisionProxy* HitProxy = new HSMECollisionProxy(EAggCollisionShape::Convex, i);
 			PDI->SetHitProxy(HitProxy);
 
 			const FColor CollisionColor = StaticMeshEditor->IsSelectedPrim(HitProxy->PrimData) ? SelectedColor : UnselectedColor;
diff --git a/Engine/Source/Editor/StaticMeshEditor/Public/IStaticMeshEditor.h b/Engine/Source/Editor/StaticMeshEditor/Public/IStaticMeshEditor.h
--- a/Engine/Source/Editor/StaticMeshEditor/Public/IStaticMeshEditor.h
+++ b/Engine/Source/Editor/StaticMeshEditor/Public/IStaticMeshEditor.h
@@ -5,7 +5,7 @@
 #include "CoreMinimal.h"
 #include "Engine/EngineBaseTypes.h"
 #include "Toolkits/AssetEditorToolkit.h"
-#include "PhysicsPublic.h"
+#include "PhysicsEngine/ShapeElem.h"
 
 class UStaticMesh;
 class UStaticMeshComponent;
@@ -27,10 +27,10 @@ public:
 	 */
 	struct FPrimData
 	{
-		EKCollisionPrimitiveType		PrimType;
+		EAggCollisionShape::Type		PrimType;
 		int32							PrimIndex;
 
-		FPrimData(EKCollisionPrimitiveType InPrimType, int32 InPrimIndex) :
+		FPrimData(EAggCollisionShape::Type InPrimType, int32 InPrimIndex) :
 			PrimType(InPrimType),
 			PrimIndex(InPrimIndex) {}
 
diff --git a/Engine/Source/Editor/UMGEditor/Private/Designer/SDesignerView.cpp b/Engine/Source/Editor/UMGEditor/Private/Designer/SDesignerView.cpp
--- a/Engine/Source/Editor/UMGEditor/Private/Designer/SDesignerView.cpp
+++ b/Engine/Source/Editor/UMGEditor/Private/Designer/SDesignerView.cpp
@@ -574,7 +574,7 @@ TSharedRef<SWidget> SDesignerView::CreateOverlayUI()
 			.ButtonStyle(FEditorStyle::Get(), "ViewportMenu.Button")
 			.ToolTipText(LOCTEXT("ZoomToFit_ToolTip", "Zoom To Fit"))
 			.OnClicked(this, &SDesignerView::HandleZoomToFitClicked)
-			.ContentPadding(1.0f)
+			.ContentPadding(FEditorStyle::Get().GetMargin("ViewportMenu.SToolBarButtonBlock.Button.Padding"))
 			[
 				SNew(SImage)
 				.Image(FEditorStyle::GetBrush("UMGEditor.ZoomToFit"))
@@ -590,7 +590,7 @@ TSharedRef<SWidget> SDesignerView::CreateOverlayUI()
 			.ButtonStyle(FEditorStyle::Get(), "ViewportMenu.Button")
 			.ForegroundColor(FLinearColor::Black)
 			.OnGetMenuContent(this, &SDesignerView::GetResolutionsMenu)
-			.ContentPadding(1.0f)
+			.ContentPadding(FEditorStyle::Get().GetMargin("ViewportMenu.SToolBarButtonBlock.Button.Padding"))
 			.ButtonContent()
 			[
 				SNew(STextBlock)
@@ -608,7 +608,7 @@ TSharedRef<SWidget> SDesignerView::CreateOverlayUI()
 			.ButtonStyle(FEditorStyle::Get(), "ViewportMenu.Button")
 			.ForegroundColor(FLinearColor::Black)
 			.OnGetMenuContent(this, &SDesignerView::GetScreenSizingFillMenu)
-			.ContentPadding(1.0f)
+			.ContentPadding(FEditorStyle::Get().GetMargin("ViewportMenu.SToolBarButtonBlock.Button.Padding"))
 			.ButtonContent()
 			[
 				SNew(STextBlock)
diff --git a/Engine/Source/Editor/UnrealEd/Classes/Animation/DebugSkelMeshComponent.h b/Engine/Source/Editor/UnrealEd/Classes/Animation/DebugSkelMeshComponent.h
--- a/Engine/Source/Editor/UnrealEd/Classes/Animation/DebugSkelMeshComponent.h
+++ b/Engine/Source/Editor/UnrealEd/Classes/Animation/DebugSkelMeshComponent.h
@@ -123,8 +123,8 @@ public:
 	}
 };
 
-UCLASS(transient, MinimalAPI)
-class UDebugSkelMeshComponent : public USkeletalMeshComponent
+UCLASS(transient)
+class UNREALED_API UDebugSkelMeshComponent : public USkeletalMeshComponent
 {
 	GENERATED_UCLASS_BODY()
 
@@ -223,6 +223,9 @@ class UDebugSkelMeshComponent : public USkeletalMeshComponent
 	UPROPERTY(transient)
 	bool bClothCullBackface;
 
+	/* Bounds computed from cloth. */
+	FBoxSphereBounds CachedClothBounds;
+
 	/** Non Compressed SpaceBases for when bDisplayRawAnimation == true **/
 	TArray<FTransform> UncompressedSpaceBases;
 
@@ -296,49 +299,55 @@ class UDebugSkelMeshComponent : public USkeletalMeshComponent
 	//~ End SkeletalMeshComponent Interface
 	// Preview.
 	// @todo document
-	UNREALED_API bool IsPreviewOn() const;
+	bool IsPreviewOn() const;
 
 	// @todo document
-	UNREALED_API FString GetPreviewText() const;
+	FString GetPreviewText() const;
 
 	// @todo anim : you still need to give asset, so that we know which one to disable
 	// we can disable per asset, so that if some other window disabled before me, I don't accidently turn it off
-	UNREALED_API void EnablePreview(bool bEnable, class UAnimationAsset * PreviewAsset);
+	void EnablePreview(bool bEnable, class UAnimationAsset * PreviewAsset);
 
 	/**
 	 * Update material information depending on color render mode 
 	 * Refresh/replace materials 
 	 */
-	UNREALED_API void SetShowBoneWeight(bool bNewShowBoneWeight);
+	void SetShowBoneWeight(bool bNewShowBoneWeight);
 
 	/**
 	* Update material information depending on color render mode
 	* Refresh/replace materials
 	*/
-	UNREALED_API void SetShowMorphTargetVerts(bool bNewShowMorphTargetVerts);
+	void SetShowMorphTargetVerts(bool bNewShowMorphTargetVerts);
 
 	/**
 	 * Does it use in-game bounds or bounds calculated from bones
 	 */
-	UNREALED_API bool IsUsingInGameBounds() const;
+	bool IsUsingInGameBounds() const;
 
 	/**
 	 * Set to use in-game bounds or bounds calculated from bones
 	 */
-	UNREALED_API void UseInGameBounds(bool bUseInGameBounds);
+	void UseInGameBounds(bool bUseInGameBounds);
 
 	/**
 	 * Test if in-game bounds are as big as preview bounds
 	 */
-	UNREALED_API bool CheckIfBoundsAreCorrrect();
+	bool CheckIfBoundsAreCorrrect();
 
 	/** 
 	 * Update components position based on animation root motion
 	 */
-	UNREALED_API void ConsumeRootMotion(const FVector& FloorMin, const FVector& FloorMax);
+	void ConsumeRootMotion(const FVector& FloorMin, const FVector& FloorMax);
 
 	/** Sets the flag used to determine whether or not the current active cloth sim mesh should be rendered */
-	UNREALED_API void SetShowClothProperty(bool bState);
+	void SetShowClothProperty(bool bState);
+
+	/** Get whether we should be previewing root motion */
+	bool GetPreviewRootMotion() const;
+
+	/** Set whether we should be previewing root motion. Note: disabling root motion preview resets transform. */
+	void SetPreviewRootMotion(bool bInPreviewRootMotion);
 
 #if WITH_EDITOR
 	//TODO - This is a really poor way to post errors to the user. Work out a better way.
@@ -371,15 +380,15 @@ class UDebugSkelMeshComponent : public USkeletalMeshComponent
 	 * if bShowOnlyClothSections is true, shows only cloth sections. On the other hand, 
 	 * if bShowOnlyClothSections is false, hides only cloth sections.
 	 */
-	UNREALED_API void ToggleClothSectionsVisibility(bool bShowOnlyClothSections);
+	void ToggleClothSectionsVisibility(bool bShowOnlyClothSections);
 	/** Restore all section visibilities to original states for all LODs */
-	UNREALED_API void RestoreClothSectionsVisibility();
+	void RestoreClothSectionsVisibility();
 
 	/** 
 	 * To normal game/runtime code we don't want to expose a non-const pointer to the simulation, so we can only get
 	 * one from this editor-only component. Intended for debug options/visualisations/editor-only code to poke the sim
 	 */
-	UNREALED_API IClothingSimulation* GetMutableClothingSimulation();
+	IClothingSimulation* GetMutableClothingSimulation();
 
 	int32 FindCurrentSectionDisplayMode();
 
@@ -395,15 +404,15 @@ class UDebugSkelMeshComponent : public USkeletalMeshComponent
 	/** The currently selected mask inside the above LOD to be painted */
 	int32 SelectedClothingLodMaskForPainting;
 
-	UNREALED_API void ToggleMeshSectionForCloth(FGuid InClothGuid);
+	void ToggleMeshSectionForCloth(FGuid InClothGuid);
 
 	// fixes up the disabled flags so clothing is enabled and originals are disabled as
 	// ToggleMeshSectionForCloth will make these get out of sync
-	UNREALED_API void ResetMeshSectionVisibility();
+	void ResetMeshSectionVisibility();
 
 	// Rebuilds the fixed parameter on the mesh to mesh data, to be used if the editor has
 	// changed a vert to be fixed or unfixed otherwise the simulation will not work
-	UNREALED_API void RebuildClothingSectionsFixedVerts();
+	void RebuildClothingSectionsFixedVerts();
 
 	TArray<FVector> SkinnedSelectedClothingPositions;
 	TArray<FVector> SkinnedSelectedClothingNormals;
@@ -417,6 +426,8 @@ private:
 	// Helper function to enable overlay material
 	void EnableOverlayMaterial(bool bEnable);
 
+	// Rebuilds the cloth bounds for the asset.
+	void RebuildCachedClothBounds();
 protected:
 
 	// Overridden to support single clothing ticks
@@ -432,7 +443,7 @@ public:
 	
 	virtual void TickComponent(float DeltaTime, enum ELevelTick TickType, FActorComponentTickFunction *ThisTickFunction) override;
 
-	void UNREALED_API RefreshSelectedClothingSkinnedPositions();
+	void RefreshSelectedClothingSkinnedPositions();
 
 	virtual void GetUsedMaterials(TArray<UMaterialInterface *>& OutMaterials, bool bGetDebugMaterials = false) const override;
 
diff --git a/Engine/Source/Editor/UnrealEd/Classes/Animation/EditorAnimCurveBoneLinks.h b/Engine/Source/Editor/UnrealEd/Classes/Animation/EditorAnimCurveBoneLinks.h
--- a/Engine/Source/Editor/UnrealEd/Classes/Animation/EditorAnimCurveBoneLinks.h
+++ b/Engine/Source/Editor/UnrealEd/Classes/Animation/EditorAnimCurveBoneLinks.h
@@ -35,8 +35,16 @@ public:
 	UPROPERTY(EditAnywhere, Category = Bones)
 	TArray<FBoneReference> ConnectedBones;
 
+	/** Max (Lowest) LOD to evaluate to curve. 
+	 *  Since LOD goes from 0 to high number, we call it Max. 
+	 *  For example, if you have 3 LODs (0, 1, 2), and if you want this to work until LOD 1, type 1.  
+	 *  Then the curve will be evaluated until LOD1, but not for LOD 2
+	 *  Default value is 255 */
+	UPROPERTY(EditAnywhere, Category = LOD)
+	uint8 MaxLOD;
+
 	virtual void PostEditChangeProperty(FPropertyChangedEvent& PropertyChangedEvent) override;
 
 	// refresh current Connected Bones data 
-	UNREALED_API void Refresh(const FSmartName& InCurveName, const TArray<FBoneReference>& CurrentLinks);
+	UNREALED_API void Refresh(const FSmartName& InCurveName, const TArray<FBoneReference>& CurrentLinks, uint8 InMaxLOD);
 };
diff --git a/Engine/Source/Editor/UnrealEd/Classes/Preferences/PersonaOptions.h b/Engine/Source/Editor/UnrealEd/Classes/Preferences/PersonaOptions.h
--- a/Engine/Source/Editor/UnrealEd/Classes/Preferences/PersonaOptions.h
+++ b/Engine/Source/Editor/UnrealEd/Classes/Preferences/PersonaOptions.h
@@ -51,9 +51,6 @@ class UNREALED_API UPersonaOptions : public UObject
 	UPROPERTY(EditAnywhere, config, Category = "Viewport")
 	TEnumAsByte<EViewModeIndex> ViewModeIndex;
 
-	UPROPERTY(EditAnywhere, config, Category = "Viewport")
-	FLinearColor ViewportBackgroundColor;
-
 	UPROPERTY(EditAnywhere, config, Category = "Viewport")
 	float ViewFOV;
 
@@ -80,6 +77,10 @@ class UNREALED_API UPersonaOptions : public UObject
 	UPROPERTY(EditAnywhere, config, Category = "Skeleton Tree")
 	bool bFlattenSkeletonHierarchyWhenFiltering;
 
+	/** Whether to hide parent items when filtering or to display them grayed out */
+	UPROPERTY(EditAnywhere, config, Category = "Skeleton Tree")
+	bool bHideParentsWhenFiltering;
+
 	UPROPERTY(EditAnywhere, config, Category = "Preview Scene")
 	bool bAllowPreviewMeshCollectionsToSelectFromDifferentSkeletons;
 
@@ -91,7 +92,6 @@ class UNREALED_API UPersonaOptions : public UObject
 	uint32 NumFolderFiltersInAssetBrowser;
 
 public:
-	void SetViewportBackgroundColor( const FLinearColor& InViewportBackgroundColor);
 	void SetShowGrid( bool bInShowGrid );
 	void SetHighlightOrigin( bool bInHighlightOrigin );
 	void SetShowFloor( bool bInShowFloor );
diff --git a/Engine/Source/Editor/UnrealEd/Classes/Settings/EditorExperimentalSettings.h b/Engine/Source/Editor/UnrealEd/Classes/Settings/EditorExperimentalSettings.h
--- a/Engine/Source/Editor/UnrealEd/Classes/Settings/EditorExperimentalSettings.h
+++ b/Engine/Source/Editor/UnrealEd/Classes/Settings/EditorExperimentalSettings.h
@@ -113,10 +113,6 @@ public:
 	UPROPERTY(EditAnywhere, config, Category = Tools, meta = (ConfigRestartRequired = true))
 	bool bFacialAnimationImporter;
 
-	/** Enable experimental clothing tools (parameter painting and simulation configuration) found in the skeletal mesh editor */
-	UPROPERTY(EditAnywhere, config, Category = Tools, meta = (ConfigRestartRequired = true))
-	bool bClothingTools;
-
 	/** Allow animation blueprints to be recompiled while a PIE session is running */
 	UPROPERTY(EditAnywhere, config, Category = Tools)
 	bool bEnableLiveRecompilationOfAnimationBlueprints;
diff --git a/Engine/Source/Editor/UnrealEd/Private/Animation/DebugSkelMeshComponent.cpp b/Engine/Source/Editor/UnrealEd/Private/Animation/DebugSkelMeshComponent.cpp
--- a/Engine/Source/Editor/UnrealEd/Private/Animation/DebugSkelMeshComponent.cpp
+++ b/Engine/Source/Editor/UnrealEd/Private/Animation/DebugSkelMeshComponent.cpp
@@ -45,13 +45,15 @@ UDebugSkelMeshComponent::UDebugSkelMeshComponent(const FObjectInitializer& Objec
 
 	bPauseClothingSimulationWithAnim = false;
 	bPerformSingleClothingTick = false;
+
+	CachedClothBounds = FBoxSphereBounds(ForceInit);
 }
 
 FBoxSphereBounds UDebugSkelMeshComponent::CalcBounds(const FTransform& LocalToWorld) const
 {
 	FBoxSphereBounds Result = Super::CalcBounds(LocalToWorld);
 
-	if (! IsUsingInGameBounds())
+	if (!IsUsingInGameBounds())
 	{
 		// extend bounds by required bones (respecting current LOD) but without root bone
 		if (GetNumComponentSpaceTransforms())
@@ -63,10 +65,18 @@ FBoxSphereBounds UDebugSkelMeshComponent::CalcBounds(const FTransform& LocalToWo
 				FBoneIndexType RequiredBoneIndex = RequiredBones[BoneIndex];
 				BoundingBox += GetBoneMatrix((int32)RequiredBoneIndex).GetOrigin();
 			}
+
 			Result = Result + FBoxSphereBounds(BoundingBox);
 		}
+
+		if ( SkeletalMesh )
+		{
+			Result = Result + SkeletalMesh->GetBounds();
+		}
 	}
 
+	Result = Result + CachedClothBounds;
+
 	return Result;
 }
 
@@ -141,7 +151,17 @@ void UDebugSkelMeshComponent::ConsumeRootMotion(const FVector& FloorMin, const F
 			SetRelativeTransform(CurrentTransform);
 		}
 	}
-	else
+}
+
+bool UDebugSkelMeshComponent::GetPreviewRootMotion() const
+{
+	return bPreviewRootMotion;
+}
+
+void UDebugSkelMeshComponent::SetPreviewRootMotion(bool bInPreviewRootMotion)
+{
+	bPreviewRootMotion = bInPreviewRootMotion;
+	if (!bPreviewRootMotion)
 	{
 		if (TurnTableMode == EPersonaTurnTableMode::Stopped)
 		{
@@ -266,6 +286,12 @@ void UDebugSkelMeshComponent::InitAnim(bool bForceReinit)
 		AnimScriptInstance = PreviewInstance;
 		AnimScriptInstance->InitializeAnimation();
 	}
+	else
+	{
+		// Make sure we initialize the preview instance here, as we want the required bones to be up to date
+		// even if we arent using the instance right now.
+		PreviewInstance->InitializeAnimation();
+	}
 
 	if(PostProcessAnimInstance)
 	{
@@ -909,6 +935,7 @@ void UDebugSkelMeshComponent::RefreshSelectedClothingSkinnedPositions()
 				FClothLODData& LodData = ConcreteAsset->LodData[SelectedClothingLodForPainting];
 
 				FClothingSimulationBase::SkinPhysicsMesh(ConcreteAsset, LodData.PhysicalMeshData, FTransform::Identity, RefToLocals.GetData(), RefToLocals.Num(), SkinnedSelectedClothingPositions, SkinnedSelectedClothingNormals);
+				RebuildCachedClothBounds();
 			}
 		}
 	}
@@ -935,6 +962,18 @@ IClothingSimulation* UDebugSkelMeshComponent::GetMutableClothingSimulation()
 	return ClothingSimulation;
 }
 
+void UDebugSkelMeshComponent::RebuildCachedClothBounds()
+{
+	FBox ClothBBox(ForceInit);
+	
+	for ( int32 Index = 0; Index < SkinnedSelectedClothingPositions.Num(); ++Index )
+	{
+		ClothBBox += SkinnedSelectedClothingPositions[Index];
+	}
+
+	CachedClothBounds = FBoxSphereBounds(ClothBBox);
+}
+
 FDebugSkelMeshSceneProxy::FDebugSkelMeshSceneProxy(const UDebugSkelMeshComponent* InComponent, FSkeletalMeshResource* InSkelMeshResource, const FColor& InWireframeOverlayColor /*= FColor::White*/) :
 	FSkeletalMeshSceneProxy(InComponent, InSkelMeshResource)
 {
diff --git a/Engine/Source/Editor/UnrealEd/Private/Animation/EditorAnimCurveBoneLinks.cpp b/Engine/Source/Editor/UnrealEd/Private/Animation/EditorAnimCurveBoneLinks.cpp
--- a/Engine/Source/Editor/UnrealEd/Private/Animation/EditorAnimCurveBoneLinks.cpp
+++ b/Engine/Source/Editor/UnrealEd/Private/Animation/EditorAnimCurveBoneLinks.cpp
@@ -10,6 +10,7 @@
 
 UEditorAnimCurveBoneLinks::UEditorAnimCurveBoneLinks(const FObjectInitializer& ObjectInitializer)
 	: Super(ObjectInitializer)
+	, MaxLOD (0xFF)
 {
 }
 
@@ -21,13 +22,14 @@ void UEditorAnimCurveBoneLinks::Initialize(const TWeakPtr<class IEditableSkeleto
 	SetFlags(RF_Transactional);
 }
 
-void UEditorAnimCurveBoneLinks::Refresh(const FSmartName& InCurveName, const TArray<FBoneReference>& CurrentLinks)
+void UEditorAnimCurveBoneLinks::Refresh(const FSmartName& InCurveName, const TArray<FBoneReference>& CurrentLinks, uint8 InMaxLOD)
 {
 	if (EditableSkeleton.IsValid())
 	{
 		// double check the name
 		CurveName = InCurveName;
 		ConnectedBones = CurrentLinks;
+		MaxLOD = InMaxLOD;
 	}
 }
 
diff --git a/Engine/Source/Editor/UnrealEd/Private/EditorViewportClient.cpp b/Engine/Source/Editor/UnrealEd/Private/EditorViewportClient.cpp
--- a/Engine/Source/Editor/UnrealEd/Private/EditorViewportClient.cpp
+++ b/Engine/Source/Editor/UnrealEd/Private/EditorViewportClient.cpp
@@ -447,6 +447,7 @@ FEditorViewportClient::~FEditorViewportClient()
 	if (bOwnsModeTools)
 	{
 		delete ModeTools;
+		ModeTools = nullptr;
 	}
 }
 
@@ -3636,9 +3637,7 @@ void FEditorViewportClient::RenderDragTool(const FSceneView* View, FCanvas* Canv
 
 FLinearColor FEditorViewportClient::GetBackgroundColor() const
 {
-	FLinearColor BackgroundColor = FColor(55, 55, 55);
-
-	return BackgroundColor;
+	return PreviewScene ? PreviewScene->GetBackgroundColor() : FColor(55, 55, 55);
 }
 
 void FEditorViewportClient::SetCameraSetup(
diff --git a/Engine/Source/Editor/UnrealEd/Private/Fbx/FbxSkeletalMeshImport.cpp b/Engine/Source/Editor/UnrealEd/Private/Fbx/FbxSkeletalMeshImport.cpp
--- a/Engine/Source/Editor/UnrealEd/Private/Fbx/FbxSkeletalMeshImport.cpp
+++ b/Engine/Source/Editor/UnrealEd/Private/Fbx/FbxSkeletalMeshImport.cpp
@@ -1773,7 +1773,6 @@ USkeletalMesh* UnFbx::FFbxImporter::ImportSkeletalMesh(FImportSkeletalMeshArgs &
 				else
 				{
 					FPhysAssetCreateParams NewBodyData;
-					NewBodyData.Initialize();
 					FText CreationErrorMessage;
 					bool bSuccess = FPhysicsAssetUtils::CreateFromSkeletalMesh(NewPhysicsAsset, SkeletalMesh, NewBodyData, CreationErrorMessage);
 					if (!bSuccess)
diff --git a/Engine/Source/Editor/UnrealEd/Private/PhysicsAssetUtils.cpp b/Engine/Source/Editor/UnrealEd/Private/PhysicsAssetUtils.cpp
--- a/Engine/Source/Editor/UnrealEd/Private/PhysicsAssetUtils.cpp
+++ b/Engine/Source/Editor/UnrealEd/Private/PhysicsAssetUtils.cpp
@@ -15,30 +15,16 @@
 #include "PhysicsEngine/PhysicsAsset.h"
 #include "PhysicsEngine/PhysicsConstraintTemplate.h"
 #include "PreviewScene.h"
+#include "ScopedSlowTask.h"
 #include "SkinnedBoneTriangleCache.h"
 
-void FPhysAssetCreateParams::Initialize()
-{
-	MinBoneSize = 20.0f;
-	MinWeldSize = KINDA_SMALL_NUMBER;
-	GeomType = EFG_Sphyl;
-	VertWeight = EVW_DominantWeight;
-	bAutoOrientToBone = true;
-	bCreateJoints = true;
-	bWalkPastSmall = true;
-	bBodyForAll = false;
-	AngularConstraintMode = ACM_Limited;
-	HullAccuracy = 0.5;
-	MaxHullVerts = 16;
-}
-
 namespace FPhysicsAssetUtils
 {
 	static const float	DefaultPrimSize = 15.0f;
 	static const float	MinPrimSize = 0.5f;
 
 	// Forward declarations
-	bool CreateCollisionFromBoneInternal(UBodySetup* bs, USkeletalMesh* skelMesh, int32 BoneIndex, FPhysAssetCreateParams& Params, const FBoneVertInfo& Info, const FSkinnedBoneTriangleCache& TriangleCache);
+	bool CreateCollisionFromBoneInternal(UBodySetup* bs, USkeletalMesh* skelMesh, int32 BoneIndex, const FPhysAssetCreateParams& Params, const FBoneVertInfo& Info, const FSkinnedBoneTriangleCache& TriangleCache);
 
 /** Returns INDEX_NONE if no children in the visual asset or if more than one parent */
 static int32 GetChildIndex(int32 BoneIndex, USkeletalMesh* SkelMesh, const TArray<FBoneVertInfo>& Infos)
@@ -128,7 +114,7 @@ void AddInfoToParentInfo(const FTransform& LocalToParentTM, const FBoneVertInfo&
 	}
 }
 
-bool CreateFromSkeletalMeshInternal(UPhysicsAsset* PhysicsAsset, USkeletalMesh* SkelMesh, FPhysAssetCreateParams& Params, const FSkinnedBoneTriangleCache& TriangleCache)
+bool CreateFromSkeletalMeshInternal(UPhysicsAsset* PhysicsAsset, USkeletalMesh* SkelMesh, const FPhysAssetCreateParams& Params, const FSkinnedBoneTriangleCache& TriangleCache)
 {
 	IMeshUtilities& MeshUtilities = FModuleManager::Get().LoadModuleChecked<IMeshUtilities>("MeshUtilities");
 
@@ -201,6 +187,9 @@ bool CreateFromSkeletalMeshInternal(UPhysicsAsset* PhysicsAsset, USkeletalMesh*
 		}
 	}
 
+	FScopedSlowTask SlowTask((float)NumBones * 2);
+	SlowTask.MakeDialog();
+
 	// Finally, iterate through all the bones and create bodies when needed
 	for (int32 BoneIndex = 0; BoneIndex < NumBones; BoneIndex++)
 	{
@@ -227,6 +216,9 @@ bool CreateFromSkeletalMeshInternal(UPhysicsAsset* PhysicsAsset, USkeletalMesh*
 		{
 			// Go ahead and make this bone physical.
 			FName BoneName = SkelMesh->RefSkeleton.GetBoneName(BoneIndex);
+
+			SlowTask.EnterProgressFrame(1.0f, FText::Format(NSLOCTEXT("PhysicsAssetEditor", "ResetCollsionStepInfo", "Generating collision for {0}"), FText::FromName(BoneName)));
+
 			int32 NewBodyIndex = CreateNewBody(PhysicsAsset, BoneName);
 			UBodySetup* NewBodySetup = PhysicsAsset->SkeletalBodySetups[NewBodyIndex];
 			check(NewBodySetup->BoneName == BoneName);
@@ -325,6 +317,9 @@ bool CreateFromSkeletalMeshInternal(UPhysicsAsset* PhysicsAsset, USkeletalMesh*
 	for(int32 BodyIdx = 0; BodyIdx < NumBodies; ++BodyIdx)
 	{
 		FBodyInstance* BodyInstance = Bodies[BodyIdx];
+
+		SlowTask.EnterProgressFrame(1.0f, FText::Format(NSLOCTEXT("PhysicsAssetEditor", "ResetCollsionStepInfoOverlaps", "Fixing overlaps for {0}"), FText::FromName(BodyInstance->BodySetup->BoneName)));
+
 		FTransform BodyTM = BodyInstance->GetUnrealWorldTransform();
 
 		for(int32 OtherBodyIdx = BodyIdx + 1; OtherBodyIdx < NumBodies; ++OtherBodyIdx)
@@ -341,7 +336,7 @@ bool CreateFromSkeletalMeshInternal(UPhysicsAsset* PhysicsAsset, USkeletalMesh*
 	return NumBodies > 0;
 }
 
-bool CreateFromSkeletalMesh(UPhysicsAsset* PhysicsAsset, USkeletalMesh* SkelMesh, FPhysAssetCreateParams& Params, FText& OutErrorMessage, bool bSetToMesh /*= true*/)
+bool CreateFromSkeletalMesh(UPhysicsAsset* PhysicsAsset, USkeletalMesh* SkelMesh, const FPhysAssetCreateParams& Params, FText& OutErrorMessage, bool bSetToMesh /*= true*/)
 {
 	PhysicsAsset->PreviewSkeletalMesh = SkelMesh;
 
@@ -358,9 +353,10 @@ bool CreateFromSkeletalMesh(UPhysicsAsset* PhysicsAsset, USkeletalMesh* SkelMesh
 	if (!bSuccess)
 	{
 		// try lower minimum bone size 
-		Params.MinBoneSize = 1.f;
+		FPhysAssetCreateParams LocalParams = Params;
+		LocalParams.MinBoneSize = 1.f;
 
-		bSuccess = CreateFromSkeletalMeshInternal(PhysicsAsset, SkelMesh, Params, TriangleCache);
+		bSuccess = CreateFromSkeletalMeshInternal(PhysicsAsset, SkelMesh, LocalParams, TriangleCache);
 
 		if(!bSuccess)
 		{
@@ -443,7 +439,7 @@ FVector ComputeEigenVector(const FMatrix& A)
 	return Bk.GetSafeNormal();
 }
 
-bool CreateCollisionFromBoneInternal(UBodySetup* bs, USkeletalMesh* skelMesh, int32 BoneIndex, FPhysAssetCreateParams& Params, const FBoneVertInfo& Info, const FSkinnedBoneTriangleCache& TriangleCache)
+bool CreateCollisionFromBoneInternal(UBodySetup* bs, USkeletalMesh* skelMesh, int32 BoneIndex, const FPhysAssetCreateParams& Params, const FBoneVertInfo& Info, const FSkinnedBoneTriangleCache& TriangleCache)
 {
 #if WITH_EDITOR
 	if (Params.GeomType != EFG_MultiConvexHull)	//multi convex hull can fail so wait to clear it
@@ -550,14 +546,21 @@ bool CreateCollisionFromBoneInternal(UBodySetup* bs, USkeletalMesh* skelMesh, in
 	}
 	else if (Params.GeomType == EFG_MultiConvexHull)
 	{
-#if WITH_EDITOR
-		bs->RemoveSimpleCollision();
-#endif
 		TArray<FVector> Verts;
 		TArray<uint32> Indices;
 		TriangleCache.GetVerticesAndIndicesForBone(BoneIndex, Verts, Indices);
 
-		DecomposeMeshToHulls(bs, Verts, Indices, Params.HullAccuracy, Params.MaxHullVerts);
+		if (Verts.Num())
+		{
+			DecomposeMeshToHulls(bs, Verts, Indices, Params.HullAccuracy, Params.MaxHullVerts);
+		}
+		else
+		{
+			FMessageLog EditorErrors("EditorErrors");
+			EditorErrors.Warning(NSLOCTEXT("PhysicsAssetUtils", "ConvexNoPositions", "Unable to create a convex hull for the given bone as there are no vertices associated with the bone."));
+			EditorErrors.Open();
+			return false;
+		}
 	}
 	else if (Params.GeomType == EFG_Sphyl)
 	{
@@ -596,7 +599,7 @@ bool CreateCollisionFromBoneInternal(UBodySetup* bs, USkeletalMesh* skelMesh, in
 	return true;
 }
 
-bool CreateCollisionFromBone(UBodySetup* bs, USkeletalMesh* skelMesh, int32 BoneIndex, FPhysAssetCreateParams& Params, const FBoneVertInfo& Info)
+bool CreateCollisionFromBone(UBodySetup* bs, USkeletalMesh* skelMesh, int32 BoneIndex, const FPhysAssetCreateParams& Params, const FBoneVertInfo& Info)
 {
 	check(skelMesh);
 
@@ -610,7 +613,7 @@ bool CreateCollisionFromBone(UBodySetup* bs, USkeletalMesh* skelMesh, int32 Bone
 	return CreateCollisionFromBoneInternal(bs, skelMesh, BoneIndex, Params, Info, TriangleCache);
 }
 
-bool CreateCollisionFromBones(UBodySetup* bs, USkeletalMesh* skelMesh, const TArray<int32>& BoneIndices, FPhysAssetCreateParams& Params, const FBoneVertInfo& Info)
+bool CreateCollisionFromBones(UBodySetup* bs, USkeletalMesh* skelMesh, const TArray<int32>& BoneIndices, const FPhysAssetCreateParams& Params, const FBoneVertInfo& Info)
 {
 	check(skelMesh);
 
diff --git a/Engine/Source/Editor/UnrealEd/Private/PreferenceStubs.cpp b/Engine/Source/Editor/UnrealEd/Private/PreferenceStubs.cpp
--- a/Engine/Source/Editor/UnrealEd/Private/PreferenceStubs.cpp
+++ b/Engine/Source/Editor/UnrealEd/Private/PreferenceStubs.cpp
@@ -6,7 +6,7 @@
 #include "Preferences/CurveEdOptions.h"
 #include "Preferences/MaterialEditorOptions.h"
 #include "Preferences/PersonaOptions.h"
-#include "Preferences/PhATSimOptions.h"
+#include "Preferences/PhysicsAssetEditorOptions.h"
 
 // @todo find a better place for all of this, preferably in the appropriate modules
 // though this would require the classes to be relocated as well
@@ -20,16 +20,35 @@ UCascadeOptions::UCascadeOptions(const FObjectInitializer& ObjectInitializer)
 }
 
 ///////////////////////////////////////////////////////////////////////////////////////////////////
-////////////////// UPhATSimOptions ////////////////////////////////////////////////////////////////
+////////////////// UPhysicsAssetEditorOptions /////////////////////////////////////////////////////
 ///////////////////////////////////////////////////////////////////////////////////////////////////
-UPhATSimOptions::UPhATSimOptions(const FObjectInitializer& ObjectInitializer)
+UPhysicsAssetEditorOptions::UPhysicsAssetEditorOptions(const FObjectInitializer& ObjectInitializer)
 	: Super(ObjectInitializer)
 {
 	PhysicsBlend = 1.0f;
 	bUpdateJointsFromAnimation = false;
-	//bImmediatePhysics = false;
 	MaxFPS = -1;
-	TimeDilation = 1.f;
+
+	// These should duplicate defaults from UPhysicsHandleComponent
+	HandleLinearDamping = 200.0f;
+	HandleLinearStiffness = 750.0f;
+	HandleAngularDamping = 500.0f;
+	HandleAngularStiffness = 1500.0f;
+	InterpolationSpeed = 50.f;
+
+	bShowConstraintsAsPoints = false;
+	ConstraintDrawSize = 1.0f;
+
+	// view options
+	MeshViewMode = EPhysicsAssetEditorRenderMode::Solid;
+	CollisionViewMode = EPhysicsAssetEditorRenderMode::Solid;
+	ConstraintViewMode = EPhysicsAssetEditorConstraintViewMode::AllLimits;
+	SimulationMeshViewMode = EPhysicsAssetEditorRenderMode::Solid;
+	SimulationCollisionViewMode = EPhysicsAssetEditorRenderMode::Solid;
+	SimulationConstraintViewMode = EPhysicsAssetEditorConstraintViewMode::None;
+
+	CollisionOpacity = 0.3f;
+	bSolidRenderingForSelectedOnly = false;
 }
 
 UMaterialEditorOptions::UMaterialEditorOptions(const FObjectInitializer& ObjectInitializer)
@@ -62,12 +81,6 @@ UPersonaOptions::UPersonaOptions(const FObjectInitializer& ObjectInitializer)
 	bUseAudioAttenuation = true;
 }
 
-void UPersonaOptions::SetViewportBackgroundColor( const FLinearColor& InViewportBackgroundColor)
-{
-	ViewportBackgroundColor = InViewportBackgroundColor;
-	SaveConfig();
-}
-
 void UPersonaOptions::SetShowGrid( bool bInShowGrid )
 {
 	bShowGrid = bInShowGrid;
diff --git a/Engine/Source/Editor/UnrealEd/Private/SEditorViewportToolBarMenu.cpp b/Engine/Source/Editor/UnrealEd/Private/SEditorViewportToolBarMenu.cpp
--- a/Engine/Source/Editor/UnrealEd/Private/SEditorViewportToolBarMenu.cpp
+++ b/Engine/Source/Editor/UnrealEd/Private/SEditorViewportToolBarMenu.cpp
@@ -24,20 +24,28 @@ void SEditorViewportToolbarMenu::Construct( const FArguments& Declaration )
 
 	TSharedPtr<SWidget> ButtonContent;
 	
+	const float MenuIconSize = 16.0f;
+
 	// Create the content for the button.  We always use an image over a label if it is valid
 	if( ImageName != NAME_None )
 	{
 		ButtonContent = 
-			SNew( SImage )	
-				.Image( ImageBrush )
-				.ColorAndOpacity( FSlateColor::UseForeground() );
+			SNew(SBox)
+			.HeightOverride(MenuIconSize)
+			.VAlign(VAlign_Center)
+			.HAlign(HAlign_Center)
+			[
+				SNew( SImage )	
+					.Image( ImageBrush )
+					.ColorAndOpacity( FSlateColor::UseForeground() )
+			];
 	}
 	else
 	{
 		if( LabelIconBrush.IsBound() || LabelIconBrush.Get() != NULL )
 		{
 			// Label with an icon to the left
-			const float MenuIconSize = 16.0f;
+			
 			ButtonContent = 
 				SNew( SHorizontalBox )
 				+SHorizontalBox::Slot()
@@ -48,6 +56,8 @@ void SEditorViewportToolbarMenu::Construct( const FArguments& Declaration )
 					.Visibility( this, &SEditorViewportToolbarMenu::GetLabelIconVisibility )
 					.WidthOverride( MenuIconSize )
 					.HeightOverride( MenuIconSize )
+					.VAlign(VAlign_Center)
+					.HAlign(HAlign_Center)
 					[
 						SNew( SImage )
 							.Image( LabelIconBrush )
@@ -57,18 +67,28 @@ void SEditorViewportToolbarMenu::Construct( const FArguments& Declaration )
 				.FillWidth( 1.0f )
 				.VAlign( VAlign_Center )
 				[
-					SNew( STextBlock )
-						.Font( FEditorStyle::GetFontStyle("EditorViewportToolBar.Font") )
-						.Text( Label )
+					SNew(SBox)
+					.HeightOverride(MenuIconSize)
+					.VAlign(VAlign_Center)
+					[
+						SNew( STextBlock )
+							.Font( FEditorStyle::GetFontStyle("EditorViewportToolBar.Font") )
+							.Text( Label )
+					]
 				];
 		}
 		else
 		{
 			// Just the label text, no icon
 			ButtonContent = 
-				SNew( STextBlock )
-					.Font( FEditorStyle::GetFontStyle("EditorViewportToolBar.Font") )
-					.Text( Label );
+				SNew(SBox)
+				.HeightOverride(MenuIconSize)
+				.VAlign(VAlign_Center)
+				[
+					SNew( STextBlock )
+						.Font( FEditorStyle::GetFontStyle("EditorViewportToolBar.Font") )
+						.Text( Label )
+				];
 		}
 	}
 
diff --git a/Engine/Source/Editor/UnrealEd/Private/SEditorViewportViewMenu.cpp b/Engine/Source/Editor/UnrealEd/Private/SEditorViewportViewMenu.cpp
--- a/Engine/Source/Editor/UnrealEd/Private/SEditorViewportViewMenu.cpp
+++ b/Engine/Source/Editor/UnrealEd/Private/SEditorViewportViewMenu.cpp
@@ -18,17 +18,17 @@ void SEditorViewportViewMenu::Construct( const FArguments& InArgs, TSharedRef<SE
 		SEditorViewportToolbarMenu::FArguments()
 			.ParentToolBar( InParentToolBar)
 			.Cursor( EMouseCursor::Default )
-			.Label( this, &SEditorViewportViewMenu::GetViewMenuLabel )
-			.LabelIcon( this, &SEditorViewportViewMenu::GetViewMenuLabelIcon )
+			.Label_Static( &SEditorViewportViewMenu::GetViewMenuLabel, Viewport)
+			.LabelIcon( this, &SEditorViewportViewMenu::GetViewMenuLabelBrush)
 			.OnGetMenuContent( this, &SEditorViewportViewMenu::GenerateViewMenuContent )
 	);
 		
 }
 
-FText SEditorViewportViewMenu::GetViewMenuLabel() const
+FText SEditorViewportViewMenu::GetViewMenuLabel(TWeakPtr<SEditorViewport> InViewport)
 {
 	FText Label = LOCTEXT("ViewMenuTitle_Default", "View");
-	TSharedPtr< SEditorViewport > PinnedViewport = Viewport.Pin();
+	TSharedPtr< SEditorViewport > PinnedViewport = InViewport.Pin();
 	if( PinnedViewport.IsValid() )
 	{
 		switch( PinnedViewport->GetViewportClient()->GetViewMode() )
@@ -128,10 +128,15 @@ FText SEditorViewportViewMenu::GetViewMenuLabel() const
 	return Label;
 }
 
-const FSlateBrush* SEditorViewportViewMenu::GetViewMenuLabelIcon() const
+const FSlateBrush* SEditorViewportViewMenu::GetViewMenuLabelBrush() const
+{
+	return GetViewMenuLabelIcon(Viewport).GetIcon();
+}
+
+FSlateIcon SEditorViewportViewMenu::GetViewMenuLabelIcon(TWeakPtr<SEditorViewport> InViewport)
 {
 	FName Icon = NAME_None;
-	TSharedPtr< SEditorViewport > PinnedViewport = Viewport.Pin();
+	TSharedPtr< SEditorViewport > PinnedViewport = InViewport.Pin();
 	if( PinnedViewport.IsValid() )
 	{
 		static FName WireframeIcon( "EditorViewport.WireframeMode" );
@@ -254,17 +259,23 @@ const FSlateBrush* SEditorViewportViewMenu::GetViewMenuLabelIcon() const
 		}
 	}
 
-	return FEditorStyle::GetBrush( Icon );
+	return FSlateIcon( "EditorStyle", Icon );
 }
 
-
 TSharedRef<SWidget> SEditorViewportViewMenu::GenerateViewMenuContent() const
 {
-	const FEditorViewportCommands& BaseViewportActions = FEditorViewportCommands::Get();
-
 	const bool bInShouldCloseWindowAfterMenuSelection = true;
 
-	FMenuBuilder ViewMenuBuilder( bInShouldCloseWindowAfterMenuSelection, Viewport.Pin()->GetCommandList(), MenuExtenders );
+	FMenuBuilder ViewMenuBuilder(bInShouldCloseWindowAfterMenuSelection, Viewport.Pin()->GetCommandList(), MenuExtenders);
+
+	GenerateViewMenu(ViewMenuBuilder, ParentToolBar);
+
+	return ViewMenuBuilder.MakeWidget();
+}
+
+void SEditorViewportViewMenu::GenerateViewMenu(FMenuBuilder& ViewMenuBuilder, TWeakPtr<SViewportToolBar> ParentToolBar )
+{
+	const FEditorViewportCommands& BaseViewportActions = FEditorViewportCommands::Get();
 	{
 		// View modes
 		{
@@ -282,7 +293,7 @@ TSharedRef<SWidget> SEditorViewportViewMenu::GenerateViewMenuContent() const
 			{
 				struct Local
 				{
-					static void BuildOptimizationMenu( FMenuBuilder& Menu, TWeakPtr< SViewportToolBar > ParentToolBar )
+					static void BuildOptimizationMenu( FMenuBuilder& Menu, TWeakPtr< SViewportToolBar > InParentToolBar )
 					{
 						const FEditorViewportCommands& BaseViewportCommands = FEditorViewportCommands::Get();
 						Menu.AddMenuEntry( BaseViewportCommands.LightComplexityMode, NAME_None, LOCTEXT("LightComplexityViewModeDisplayName", "Light Complexity") );
@@ -302,20 +313,20 @@ TSharedRef<SWidget> SEditorViewportViewMenu::GenerateViewMenuContent() const
 						Menu.AddMenuEntry( BaseViewportCommands.LODColorationMode, NAME_None, LOCTEXT("LODColorationViewModeDisplayName", "LOD Coloration") );
 
 						Menu.BeginSection("TextureStreaming", LOCTEXT("TextureStreamingHeader", "Texture Streaming Accuracy") );
-						if ( AllowDebugViewShaderMode(DVSM_PrimitiveDistanceAccuracy) && (!ParentToolBar.IsValid() || ParentToolBar.Pin()->IsViewModeSupported(VMI_PrimitiveDistanceAccuracy)) )
+						if ( AllowDebugViewShaderMode(DVSM_PrimitiveDistanceAccuracy) && (!InParentToolBar.IsValid() || InParentToolBar.Pin()->IsViewModeSupported(VMI_PrimitiveDistanceAccuracy)) )
 						{
 							Menu.AddMenuEntry( BaseViewportCommands.TexStreamAccPrimitiveDistanceMode, NAME_None, LOCTEXT("TexStreamAccPrimitiveDistanceViewModeDisplayName", "Primitive Distance") );
 						}
-						if ( AllowDebugViewShaderMode(DVSM_MeshUVDensityAccuracy) && (!ParentToolBar.IsValid() || ParentToolBar.Pin()->IsViewModeSupported(VMI_MeshUVDensityAccuracy)) )
+						if ( AllowDebugViewShaderMode(DVSM_MeshUVDensityAccuracy) && (!InParentToolBar.IsValid() || InParentToolBar.Pin()->IsViewModeSupported(VMI_MeshUVDensityAccuracy)) )
 						{
 							Menu.AddMenuEntry(BaseViewportCommands.TexStreamAccMeshUVDensityMode, NAME_None, LOCTEXT("TexStreamAccMeshUVDensityViewModeDisplayName", "Mesh UV Densities"));
 						}
 						// TexCoordScale accuracy viewmode requires shaders that are only built in the TextureStreamingBuild, which requires the new metrics to be enabled.
-						if ( AllowDebugViewShaderMode(DVSM_MaterialTextureScaleAccuracy) && CVarStreamingUseNewMetrics.GetValueOnAnyThread() != 0 && (!ParentToolBar.IsValid() || ParentToolBar.Pin()->IsViewModeSupported(VMI_MaterialTextureScaleAccuracy)) )
+						if ( AllowDebugViewShaderMode(DVSM_MaterialTextureScaleAccuracy) && CVarStreamingUseNewMetrics.GetValueOnAnyThread() != 0 && (!InParentToolBar.IsValid() || InParentToolBar.Pin()->IsViewModeSupported(VMI_MaterialTextureScaleAccuracy)) )
 						{
 							Menu.AddMenuEntry(BaseViewportCommands.TexStreamAccMaterialTextureScaleMode, NAME_None, LOCTEXT("TexStreamAccMaterialTextureScaleViewModeDisplayName", "Material Texture Scales"));
 						}
-						if ( AllowDebugViewShaderMode(DVSM_RequiredTextureResolution) && (!ParentToolBar.IsValid() || ParentToolBar.Pin()->IsViewModeSupported(VMI_MaterialTextureScaleAccuracy)) )
+						if ( AllowDebugViewShaderMode(DVSM_RequiredTextureResolution) && (!InParentToolBar.IsValid() || InParentToolBar.Pin()->IsViewModeSupported(VMI_MaterialTextureScaleAccuracy)) )
 						{
 							Menu.AddMenuEntry(BaseViewportCommands.RequiredTextureResolutionMode, NAME_None, LOCTEXT("RequiredTextureResolutionModeDisplayName", "Required Texture Resolution"));
 						}
@@ -355,9 +366,6 @@ TSharedRef<SWidget> SEditorViewportViewMenu::GenerateViewMenuContent() const
 			ViewMenuBuilder.EndSection();
 		}
 	}
-
-	return ViewMenuBuilder.MakeWidget();
-	
 }
 
 #undef LOCTEXT_NAMESPACE
diff --git a/Engine/Source/Editor/UnrealEd/Private/SViewportToolBarComboMenu.cpp b/Engine/Source/Editor/UnrealEd/Private/SViewportToolBarComboMenu.cpp
--- a/Engine/Source/Editor/UnrealEd/Private/SViewportToolBarComboMenu.cpp
+++ b/Engine/Source/Editor/UnrealEd/Private/SViewportToolBarComboMenu.cpp
@@ -46,7 +46,7 @@ void SViewportToolBarComboMenu::Construct( const FArguments& InArgs )
 		TSharedRef<SWidget> ButtonContents =
 			SNew(SButton)
 			.ButtonStyle( FEditorStyle::Get(), EMultiBlockLocation::ToName(ButtonStyle,EMultiBlockLocation::End) )
-			.ContentPadding( FMargin( 5.0f, 0.0f ) )
+			.ContentPadding( FMargin( 5.0f, 2.0f ) )
 			.ToolTipText(InArgs._MenuButtonToolTip)
 			.OnClicked(this, &SViewportToolBarComboMenu::OnMenuClicked)
 			[
diff --git a/Engine/Source/Editor/UnrealEd/Private/SViewportToolBarIconMenu.cpp b/Engine/Source/Editor/UnrealEd/Private/SViewportToolBarIconMenu.cpp
--- a/Engine/Source/Editor/UnrealEd/Private/SViewportToolBarIconMenu.cpp
+++ b/Engine/Source/Editor/UnrealEd/Private/SViewportToolBarIconMenu.cpp
@@ -23,7 +23,7 @@ void SViewportToolBarIconMenu::Construct( const FArguments& InArgs )
 	[
 		SNew(SButton)
 		.ButtonStyle( FEditorStyle::Get(), ButtonStyle )
-		.ContentPadding( FMargin( 5.0f, 0.0f ) )
+		.ContentPadding( FMargin( 5.0f, 2.0f ) )
 		.OnClicked(this, &SViewportToolBarIconMenu::OnMenuClicked)
 		[
 			SNew(SHorizontalBox)
diff --git a/Engine/Source/Editor/UnrealEd/Private/ThumbnailHelpers.cpp b/Engine/Source/Editor/UnrealEd/Private/ThumbnailHelpers.cpp
--- a/Engine/Source/Editor/UnrealEd/Private/ThumbnailHelpers.cpp
+++ b/Engine/Source/Editor/UnrealEd/Private/ThumbnailHelpers.cpp
@@ -22,6 +22,7 @@
 #include "Animation/AnimSingleNodeInstance.h"
 #include "Engine/TextureCube.h"
 #include "Animation/BlendSpace1D.h"
+#include "PhysicsEngine/PhysicsAsset.h"
 
 /*
 ***************************************************************
@@ -919,6 +920,81 @@ void FAnimBlueprintThumbnailScene::GetViewMatrixParameters(const float InFOVDegr
 	OutOrbitZoom = TargetDistance + ThumbnailInfo->OrbitZoom;
 }
 
+/*
+***************************************************************
+FPhysicsAssetThumbnailScene
+***************************************************************
+*/
+
+FPhysicsAssetThumbnailScene::FPhysicsAssetThumbnailScene()
+	: FThumbnailPreviewScene()
+{
+	bForceAllUsedMipsResident = false;
+	// Create preview actor
+	// checked
+	FActorSpawnParameters SpawnInfo;
+	SpawnInfo.SpawnCollisionHandlingOverride = ESpawnActorCollisionHandlingMethod::AlwaysSpawn;
+	SpawnInfo.bNoFail = true;
+	SpawnInfo.ObjectFlags = RF_Transient;
+	PreviewActor = GetWorld()->SpawnActor<ASkeletalMeshActor>( SpawnInfo );
+
+	PreviewActor->SetActorEnableCollision(false);
+}
+
+void FPhysicsAssetThumbnailScene::SetPhysicsAsset(UPhysicsAsset* InPhysicsAsset)
+{
+	PreviewActor->GetSkeletalMeshComponent()->OverrideMaterials.Empty();
+	PreviewActor->SetActorEnableCollision(true);
+
+	if(InPhysicsAsset)
+	{
+		USkeletalMesh* SkeletalMesh = InPhysicsAsset->PreviewSkeletalMesh.LoadSynchronous();
+		if (SkeletalMesh)
+		{
+			PreviewActor->GetSkeletalMeshComponent()->SetSkeletalMesh(SkeletalMesh);
+
+			FTransform MeshTransform = FTransform::Identity;
+
+			PreviewActor->SetActorLocation(FVector(0,0,0), false);
+			PreviewActor->GetSkeletalMeshComponent()->UpdateBounds();
+
+			// Center the mesh at the world origin then offset to put it on top of the plane
+			const float BoundsZOffset = GetBoundsZOffset(PreviewActor->GetSkeletalMeshComponent()->Bounds);
+			PreviewActor->SetActorLocation( -PreviewActor->GetSkeletalMeshComponent()->Bounds.Origin + FVector(0, 0, BoundsZOffset), false );
+			PreviewActor->GetSkeletalMeshComponent()->RecreateRenderState_Concurrent();
+		}
+	}
+}
+
+void FPhysicsAssetThumbnailScene::GetViewMatrixParameters(const float InFOVDegrees, FVector& OutOrigin, float& OutOrbitPitch, float& OutOrbitYaw, float& OutOrbitZoom) const
+{
+	check(PreviewActor->GetSkeletalMeshComponent());
+
+	const float HalfFOVRadians = FMath::DegreesToRadians<float>(InFOVDegrees) * 0.5f;
+	// No need to add extra size to view slightly outside of the sphere to compensate for perspective since skeletal meshes already buffer bounds.
+	const float HalfMeshSize = PreviewActor->GetSkeletalMeshComponent()->Bounds.SphereRadius; 
+	const float BoundsZOffset = GetBoundsZOffset(PreviewActor->GetSkeletalMeshComponent()->Bounds);
+	const float TargetDistance = HalfMeshSize / FMath::Tan(HalfFOVRadians);
+
+	USceneThumbnailInfo* ThumbnailInfo = USceneThumbnailInfo::StaticClass()->GetDefaultObject<USceneThumbnailInfo>();
+	if(PreviewActor->GetSkeletalMeshComponent()->SkeletalMesh && PreviewActor->GetSkeletalMeshComponent()->SkeletalMesh->PhysicsAsset)
+	{
+		if ( USceneThumbnailInfo* InteralThumbnailInfo = Cast<USceneThumbnailInfo>(PreviewActor->GetSkeletalMeshComponent()->SkeletalMesh->PhysicsAsset->ThumbnailInfo) )
+		{
+			ThumbnailInfo = InteralThumbnailInfo;
+			if ( TargetDistance + InteralThumbnailInfo->OrbitZoom < 0 )
+			{
+				InteralThumbnailInfo->OrbitZoom = -TargetDistance;
+			}
+		}
+	}
+
+	OutOrigin = FVector(0, 0, -BoundsZOffset);
+	OutOrbitPitch = ThumbnailInfo->OrbitPitch;
+	OutOrbitYaw = ThumbnailInfo->OrbitYaw;
+	OutOrbitZoom = TargetDistance + ThumbnailInfo->OrbitZoom;
+}
+
 /*
 ***************************************************************
   FClassActorThumbnailScene
diff --git a/Engine/Source/Editor/UnrealEd/Private/ThumbnailRendering/PhysicsAssetThumbnailRenderer.cpp b/Engine/Source/Editor/UnrealEd/Private/ThumbnailRendering/PhysicsAssetThumbnailRenderer.cpp
new file mode 100644
--- /dev/null
+++ b/Engine/Source/Editor/UnrealEd/Private/ThumbnailRendering/PhysicsAssetThumbnailRenderer.cpp
@@ -0,0 +1,54 @@
+// Copyright 1998-2017 Epic Games, Inc. All Rights Reserved.
+
+#include "ThumbnailRendering/PhysicsAssetThumbnailRenderer.h"
+#include "Misc/App.h"
+#include "ShowFlags.h"
+#include "SceneView.h"
+#include "PhysicsEngine/PhysicsAsset.h"
+#include "ThumbnailHelpers.h"
+
+// FPreviewScene derived helpers for rendering
+#include "RendererInterface.h"
+#include "EngineModule.h"
+
+UPhysicsAssetThumbnailRenderer::UPhysicsAssetThumbnailRenderer(const FObjectInitializer& ObjectInitializer)
+	: Super(ObjectInitializer)
+{
+	ThumbnailScene = nullptr;
+}
+
+void UPhysicsAssetThumbnailRenderer::Draw(UObject* Object, int32 X, int32 Y, uint32 Width, uint32 Height, FRenderTarget* RenderTarget, FCanvas* Canvas)
+{
+	UPhysicsAsset* PhysicsAsset = Cast<UPhysicsAsset>(Object);
+	if (PhysicsAsset != nullptr)
+	{
+		if ( ThumbnailScene == nullptr )
+		{
+			ThumbnailScene = new FPhysicsAssetThumbnailScene();
+		}
+
+		ThumbnailScene->SetPhysicsAsset(PhysicsAsset);
+		FSceneViewFamilyContext ViewFamily( FSceneViewFamily::ConstructionValues( RenderTarget, ThumbnailScene->GetScene(), FEngineShowFlags(ESFIM_Game) )
+			.SetWorldTimes(FApp::GetCurrentTime() - GStartTime, FApp::GetDeltaTime(), FApp::GetCurrentTime() - GStartTime));
+
+		ViewFamily.EngineShowFlags.DisableAdvancedFeatures();
+		ViewFamily.EngineShowFlags.MotionBlur = 0;
+		ViewFamily.EngineShowFlags.LOD = 0;
+		ViewFamily.EngineShowFlags.Collision = 1;
+
+		ThumbnailScene->GetView(&ViewFamily, X, Y, Width, Height);
+		GetRendererModule().BeginRenderingViewFamily(Canvas,&ViewFamily);
+		ThumbnailScene->SetPhysicsAsset(nullptr);
+	}
+}
+
+void UPhysicsAssetThumbnailRenderer::BeginDestroy()
+{
+	if ( ThumbnailScene != nullptr )
+	{
+		delete ThumbnailScene;
+		ThumbnailScene = nullptr;
+	}
+
+	Super::BeginDestroy();
+}
diff --git a/Engine/Source/Editor/UnrealEd/Public/PhysicsAssetUtils.h b/Engine/Source/Editor/UnrealEd/Public/PhysicsAssetUtils.h
--- a/Engine/Source/Editor/UnrealEd/Public/PhysicsAssetUtils.h
+++ b/Engine/Source/Editor/UnrealEd/Public/PhysicsAssetUtils.h
@@ -4,6 +4,8 @@
 
 #include "CoreMinimal.h"
 #include "SkeletalMeshTypes.h"
+#include "Engine/EngineTypes.h"
+#include "PhysicsAssetUtils.generated.h"
 
 class UBodySetup;
 class UPhysicsAsset;
@@ -11,37 +13,87 @@ class UPhysicsConstraintTemplate;
 class USkeletalMesh;
 class USkeletalMeshComponent;
 
+UENUM()
 enum EPhysAssetFitGeomType
 {
-	EFG_Box,
-	EFG_Sphyl,
-	EFG_Sphere,
-	EFG_SingleConvexHull,
-	EFG_MultiConvexHull
+	EFG_Box					UMETA(DisplayName="Box"),
+	EFG_Sphyl				UMETA(DisplayName="Capsule"),
+	EFG_Sphere				UMETA(DisplayName="Sphere"),
+	EFG_SingleConvexHull	UMETA(DisplayName="Single Convex Hull"),
+	EFG_MultiConvexHull		UMETA(DisplayName="Multi Convex Hull")
 };
 
+UENUM()
 enum EPhysAssetFitVertWeight
 {
-	EVW_AnyWeight,
-	EVW_DominantWeight
+	EVW_AnyWeight			UMETA(DisplayName="Any Weight"),
+	EVW_DominantWeight		UMETA(DisplayName="Dominant Weight"),
 };
 
 /** Parameters for PhysicsAsset creation */
+USTRUCT()
 struct FPhysAssetCreateParams
 {
-	float								MinBoneSize;
-	float								MinWeldSize;
-	EPhysAssetFitGeomType				GeomType;
-	EPhysAssetFitVertWeight				VertWeight;
-	bool								bAutoOrientToBone;
-	bool								bCreateJoints;
-	bool								bWalkPastSmall;
-	bool								bBodyForAll;
-	EAngularConstraintMotion			AngularConstraintMode;
-	float								HullAccuracy;
-	int32								MaxHullVerts;
-
-	UNREALED_API void Initialize();
+	GENERATED_BODY()
+
+	FPhysAssetCreateParams()
+	{
+		MinBoneSize = 20.0f;
+		MinWeldSize = KINDA_SMALL_NUMBER;
+		GeomType = EFG_Sphyl;
+		VertWeight = EVW_DominantWeight;
+		bAutoOrientToBone = true;
+		bCreateJoints = true;
+		bWalkPastSmall = true;
+		bBodyForAll = false;
+		AngularConstraintMode = ACM_Limited;
+		HullAccuracy = 0.5f;
+		MaxHullVerts = 16;
+	}
+
+	/** Bones that are shorter than this value will be ignored for body creation */
+	UPROPERTY(EditAnywhere, Category = "Body Creation")
+	float MinBoneSize;
+
+	/** Bones that are smaller than this value will be merged together for body creation */
+	UPROPERTY(EditAnywhere, AdvancedDisplay, Category = "Body Creation")
+	float MinWeldSize;
+
+	/** The geometry type that should be used when creating bodies */
+	UPROPERTY(EditAnywhere, Category = "Body Creation", meta=(DisplayName="Primitive Type"))
+	TEnumAsByte<EPhysAssetFitGeomType> GeomType;
+
+	/** How vertices are mapped bones when approximating them with bodies */
+	UPROPERTY(EditAnywhere, Category = "Body Creation", meta=(DisplayName="Vertex Weighting Type"))
+	TEnumAsByte<EPhysAssetFitVertWeight> VertWeight;
+
+	/** Whether to automatically orient the created bodies to their corresponding bones */
+	UPROPERTY(EditAnywhere, Category = "Body Creation")
+	bool bAutoOrientToBone;
+
+	/** Whether to create constraints between adjacent created bodies */
+	UPROPERTY(EditAnywhere, Category = "Constraint Creation")
+	bool bCreateJoints;
+
+	/** Whether to skip small bones entirely (rather than merge them with adjacent bones) */
+	UPROPERTY(EditAnywhere, Category = "Body Creation", meta=(DisplayName="Walk Past Small Bones"))
+	bool bWalkPastSmall;
+
+	/** Forces creation of a body for each bone */
+	UPROPERTY(EditAnywhere, Category = "Body Creation", meta=(DisplayName="Create Body for All Bones"))
+	bool bBodyForAll;
+
+	/** The type of angular constraint to create between bodies */
+	UPROPERTY(EditAnywhere, Category = "Constraint Creation", meta=(EditCondition="bCreateJoints"))
+	TEnumAsByte<EAngularConstraintMotion> AngularConstraintMode;
+
+	/** When creating convex hulls, the target accuracy of the created hull */
+	UPROPERTY(EditAnywhere, AdvancedDisplay, Category = "Body Creation")
+	float HullAccuracy;
+
+	/** When creating convex hulls, the maximum verts that should be created */
+	UPROPERTY(EditAnywhere, AdvancedDisplay, Category = "Body Creation")
+	int32 MaxHullVerts;
 };
 
 class UPhysicsAsset;
@@ -61,7 +113,7 @@ namespace FPhysicsAssetUtils
 	 * @param	OutErrorMessage		Additional error information
 	 * @param	bSetToMesh			Whether or not to apply the physics asset to SkelMesh immediately
 	 */
-	UNREALED_API bool CreateFromSkeletalMesh(UPhysicsAsset* PhysicsAsset, USkeletalMesh* SkelMesh, FPhysAssetCreateParams& Params, FText& OutErrorMessage, bool bSetToMesh = true);
+	UNREALED_API bool CreateFromSkeletalMesh(UPhysicsAsset* PhysicsAsset, USkeletalMesh* SkelMesh, const FPhysAssetCreateParams& Params, FText& OutErrorMessage, bool bSetToMesh = true);
 
 	/** Replaces any collision already in the BodySetup with an auto-generated one using the parameters provided.
 	 * 
@@ -76,7 +128,7 @@ namespace FPhysicsAssetUtils
 	 * @param	Info				The vertices to create the collision for
 	 * @return  Returns true if successfully created collision from bone
 	 */
-	UNREALED_API bool CreateCollisionFromBone( UBodySetup* bs, USkeletalMesh* skelMesh, int32 BoneIndex, FPhysAssetCreateParams& Params, const FBoneVertInfo& Info );
+	UNREALED_API bool CreateCollisionFromBone( UBodySetup* bs, USkeletalMesh* skelMesh, int32 BoneIndex, const FPhysAssetCreateParams& Params, const FBoneVertInfo& Info );
 
 	/** Replaces any collision already in the BodySetup with an auto-generated one using the parameters provided.
 	 * 
diff --git a/Engine/Source/Editor/UnrealEd/Public/SEditorViewportViewMenu.h b/Engine/Source/Editor/UnrealEd/Public/SEditorViewportViewMenu.h
--- a/Engine/Source/Editor/UnrealEd/Public/SEditorViewportViewMenu.h
+++ b/Engine/Source/Editor/UnrealEd/Public/SEditorViewportViewMenu.h
@@ -21,11 +21,15 @@ public:
 
 	void Construct( const FArguments& InArgs, TSharedRef<SEditorViewport> InViewport, TSharedRef<class SViewportToolBar> InParentToolBar );
 
-private:
-	FText GetViewMenuLabel() const;
-	const FSlateBrush* GetViewMenuLabelIcon() const;
+	/** Generate the view menu */
+	static void GenerateViewMenu(FMenuBuilder& InMenuBuilder, TWeakPtr<SViewportToolBar> ParentToolBar);
+
+	static FText GetViewMenuLabel(TWeakPtr<SEditorViewport> InViewport);
+
+	static FSlateIcon GetViewMenuLabelIcon(TWeakPtr<SEditorViewport> InViewport);
 
 protected:
+	const FSlateBrush* GetViewMenuLabelBrush() const;
 	virtual TSharedRef<SWidget> GenerateViewMenuContent() const;
 	TWeakPtr<SEditorViewport> Viewport;
 
diff --git a/Engine/Source/Editor/UnrealEd/Public/ThumbnailHelpers.h b/Engine/Source/Editor/UnrealEd/Public/ThumbnailHelpers.h
--- a/Engine/Source/Editor/UnrealEd/Public/ThumbnailHelpers.h
+++ b/Engine/Source/Editor/UnrealEd/Public/ThumbnailHelpers.h
@@ -197,6 +197,24 @@ private:
 	class UAnimBlueprint* PreviewBlueprint;
 };
 
+class UNREALED_API FPhysicsAssetThumbnailScene : public FThumbnailPreviewScene
+{
+public:
+	/** Constructor */
+	FPhysicsAssetThumbnailScene();
+
+	/** Sets the skeletal mesh to use in the next GetView() */
+	void SetPhysicsAsset(class UPhysicsAsset* InPhysicsAsset);
+
+protected:
+	// FThumbnailPreviewScene implementation
+	virtual void GetViewMatrixParameters(const float InFOVDegrees, FVector& OutOrigin, float& OutOrbitPitch, float& OutOrbitYaw, float& OutOrbitZoom) const override;
+
+private:
+	/** The skeletal mesh actor used to display all physics asset thumbnails */
+	class ASkeletalMeshActor* PreviewActor;
+};
+
 class UActorComponent;
 
 class UNREALED_API FClassActorThumbnailScene : public FThumbnailPreviewScene
diff --git a/Engine/Source/Editor/UnrealEd/Public/UnrealEd.h b/Engine/Source/Editor/UnrealEd/Public/UnrealEd.h
--- a/Engine/Source/Editor/UnrealEd/Public/UnrealEd.h
+++ b/Engine/Source/Editor/UnrealEd/Public/UnrealEd.h
@@ -220,7 +220,7 @@ MONOLITHIC_HEADER_BOILERPLATE()
 #include "Classes/Exporters/ObjectExporterT3D.h"
 #include "Classes/Commandlets/ParticleSystemAuditCommandlet.h"
 #include "Classes/Preferences/PersonaOptions.h"
-#include "Classes/Preferences/PhATSimOptions.h"
+#include "Classes/Preferences/PhysicsAssetEditorOptions.h"
 #include "Classes/Commandlets/PkgInfoCommandlet.h"
 #include "Classes/Exporters/PolysExporterOBJ.h"
 #include "Classes/Exporters/PolysExporterT3D.h"
diff --git a/Engine/Source/Editor/UnrealEd/UnrealEd.Build.cs b/Engine/Source/Editor/UnrealEd/UnrealEd.Build.cs
--- a/Engine/Source/Editor/UnrealEd/UnrealEd.Build.cs
+++ b/Engine/Source/Editor/UnrealEd/UnrealEd.Build.cs
@@ -61,7 +61,8 @@ public class UnrealEd : ModuleRules
 				"ViewportInteraction",
 				"VREditor",
 				"Persona",
-				"ClothingSystemEditorInterface",
+                "PhysicsAssetEditor",
+                "ClothingSystemEditorInterface",
 			}
 		);
 
@@ -185,7 +186,7 @@ public class UnrealEd : ModuleRules
 				"LevelEditor",
 				"PackagesDialog",
 				"Persona",
-				"PhAT",
+                "PhysicsAssetEditor",
 				"ProjectLauncher",
 				"DeviceManager",
 				"SettingsEditor",
@@ -263,12 +264,6 @@ public class UnrealEd : ModuleRules
 			PublicDependencyModuleNames.Add("XAudio2");
 			PublicDependencyModuleNames.Add("AudioMixerXAudio2");
 
-			// SDL ony has 64 bit binary
-			if (Target.Platform == UnrealTargetPlatform.Win64)
-			{
-				PublicDependencyModuleNames.Add("AudioMixerSDL");
-			}
-
 			AddEngineThirdPartyPrivateStaticDependencies(Target,
 				"UEOgg",
 				"Vorbis",
@@ -280,6 +275,7 @@ public class UnrealEd : ModuleRules
 		if (Target.Platform == UnrealTargetPlatform.HTML5)
 		{
 			PublicDependencyModuleNames.Add("ALAudio");
+            PublicDependencyModuleNames.Add("AudioMixerSDL");
 		}
 
 		AddEngineThirdPartyPrivateStaticDependencies(Target,
diff --git a/Engine/Source/Runtime/ALAudio/Private/ALAudioSource.cpp b/Engine/Source/Runtime/ALAudio/Private/ALAudioSource.cpp
--- a/Engine/Source/Runtime/ALAudio/Private/ALAudioSource.cpp
+++ b/Engine/Source/Runtime/ALAudio/Private/ALAudioSource.cpp
@@ -103,13 +103,12 @@ void FALSoundSource::Update( void )
 	}
 	else
 	{
-		Volume = WaveInstance->Volume * WaveInstance->VolumeMultiplier;
+		Volume = WaveInstance->GetActualVolume();
 		if (SetStereoBleed())
 		{
 			// Emulate the bleed to rear speakers followed by stereo fold down
 			Volume *= 1.25f;
 		}
-		Volume *= FApp::GetVolumeMultiplier();
 		Volume *= AudioDevice->GetPlatformAudioHeadroom();
 		Volume = FMath::Clamp(Volume, 0.0f, MAX_VOLUME);
 	}
diff --git a/Engine/Source/Runtime/Android/AndroidAudio/Private/AndroidAudioSource.cpp b/Engine/Source/Runtime/Android/AndroidAudio/Private/AndroidAudioSource.cpp
--- a/Engine/Source/Runtime/Android/AndroidAudio/Private/AndroidAudioSource.cpp
+++ b/Engine/Source/Runtime/Android/AndroidAudio/Private/AndroidAudioSource.cpp
@@ -269,7 +269,7 @@ bool FSLESSoundSource::Init( FWaveInstance* InWaveInstance )
 	FSoundSource::InitCommon();
 
 	// don't do anything if no volume! THIS APPEARS TO HAVE THE VOLUME IN TIME, CHECK HERE THOUGH IF ISSUES
-	if( InWaveInstance && ( InWaveInstance->Volume * InWaveInstance->VolumeMultiplier ) <= 0 )
+	if( InWaveInstance && ( InWaveInstance->GetActualVolume()) <= 0 )
 	{
 		return false;
 	}
@@ -424,7 +424,7 @@ void FSLESSoundSource::Update( void )
 
 	FSoundSource::UpdateCommon();
 	
-	float Volume = WaveInstance->Volume * WaveInstance->VolumeMultiplier;
+	float Volume = WaveInstance->GetActualVolume();
 	if (SetStereoBleed())
 	{
 		// Emulate the bleed to rear speakers followed by stereo fold down
diff --git a/Engine/Source/Runtime/Android/AndroidRuntimeSettings/Classes/AndroidRuntimeSettings.h b/Engine/Source/Runtime/Android/AndroidRuntimeSettings/Classes/AndroidRuntimeSettings.h
--- a/Engine/Source/Runtime/Android/AndroidRuntimeSettings/Classes/AndroidRuntimeSettings.h
+++ b/Engine/Source/Runtime/Android/AndroidRuntimeSettings/Classes/AndroidRuntimeSettings.h
@@ -405,6 +405,18 @@ public:
 	UPROPERTY(config, EditAnywhere, Category = "Audio", meta = (ClampMin = "0", UIMin = "0", DisplayName = "Number of Source Workers"))
 	int32 AudioNumSourceWorkers;
 
+	/** Which of the currently enabled spatialization plugins to use on Windows. */
+	UPROPERTY(config, EditAnywhere, Category = "Audio")
+	FString SpatializationPlugin;
+
+	/** Which of the currently enabled reverb plugins to use on Windows. */
+	UPROPERTY(config, EditAnywhere, Category = "Audio")
+	FString ReverbPlugin;
+
+	/** Which of the currently enabled occlusion plugins to use on Windows. */
+	UPROPERTY(config, EditAnywhere, Category = "Audio")
+	FString OcclusionPlugin;
+
 	// Several Android graphics debuggers require configuration changes to be made to your application in order to operate. Choosing an option from this menu will configure your project to work with that graphics debugger. 
 	UPROPERTY(GlobalConfig, EditAnywhere, Category = GraphicsDebugger)
 	TEnumAsByte<EAndroidGraphicsDebugger::Type> AndroidGraphicsDebugger;
diff --git a/Engine/Source/Runtime/Android/AudioMixerAndroid/Private/AudioMixerPlatformAndroid.cpp b/Engine/Source/Runtime/Android/AudioMixerAndroid/Private/AudioMixerPlatformAndroid.cpp
--- a/Engine/Source/Runtime/Android/AudioMixerAndroid/Private/AudioMixerPlatformAndroid.cpp
+++ b/Engine/Source/Runtime/Android/AudioMixerAndroid/Private/AudioMixerPlatformAndroid.cpp
@@ -327,6 +327,30 @@ namespace Audio
 		return PlatformSettings;
 	}
 
+	static bool bSuspended = false;
+
+	void FMixerPlatformAndroid::SuspendContext()
+	{
+		if (!bSuspended)
+		{
+			bSuspended = true;
+			// set the player's state to paused
+			SLresult result = (*SL_PlayerPlayInterface)->SetPlayState(SL_PlayerPlayInterface, SL_PLAYSTATE_PAUSED);
+			check(SL_RESULT_SUCCESS == result);
+		}
+	}
+
+	void FMixerPlatformAndroid::ResumeContext()
+	{
+		// set the player's state to paused
+		if (bSuspended)
+		{
+			bSuspended = false;
+			SLresult result = (*SL_PlayerPlayInterface)->SetPlayState(SL_PlayerPlayInterface, SL_PLAYSTATE_PLAYING);
+			check(SL_RESULT_SUCCESS == result);
+		}
+	}
+
 	void FMixerPlatformAndroid::SubmitBuffer(const uint8* Buffer)
 	{
 		SLresult Result = (*SL_PlayerBufferQueue)->Enqueue(SL_PlayerBufferQueue, Buffer, AudioStreamInfo.NumOutputFrames * AudioStreamInfo.DeviceInfo.NumChannels * sizeof(int16));
@@ -373,24 +397,6 @@ namespace Audio
 		return FString();
 	}
 
-	void FMixerPlatformAndroid::ResumeContext()
-	{
-		if (bSuspended)
-		{
-			UE_LOG(LogAudioMixerAndroid, Display, TEXT("Resuming Audio"));
-			bSuspended = false;
-		}
-	}
-	
-	void FMixerPlatformAndroid::SuspendContext()
-	{
-		if (!bSuspended)
-		{
-			UE_LOG(LogAudioMixerAndroid, Display, TEXT("Suspending Audio"));
-			bSuspended = true;
-		}
-	}
-
 	void FMixerPlatformAndroid::OpenSLBufferQueueCallback(SLAndroidSimpleBufferQueueItf InQueueInterface, void* pContext)
 	{
 		FMixerPlatformAndroid* MixerPlatformAndroid = (FMixerPlatformAndroid*)pContext;
diff --git a/Engine/Source/Runtime/Android/AudioMixerAndroid/Private/AudioMixerPlatformAndroid.h b/Engine/Source/Runtime/Android/AudioMixerAndroid/Private/AudioMixerPlatformAndroid.h
--- a/Engine/Source/Runtime/Android/AudioMixerAndroid/Private/AudioMixerPlatformAndroid.h
+++ b/Engine/Source/Runtime/Android/AudioMixerAndroid/Private/AudioMixerPlatformAndroid.h
@@ -37,11 +37,9 @@ namespace Audio
 		virtual ICompressedAudioInfo* CreateCompressedAudioInfo(USoundWave* InSoundWave) override;
 		virtual FString GetDefaultDeviceName() override;
 		virtual FAudioPlatformSettings GetPlatformSettings() const override;
+		virtual void SuspendContext() override;
+		virtual void ResumeContext() override;
 		//~ End IAudioMixerPlatformInterface
-
-		// These are not being used yet but they should be in the near future
-		void ResumeContext();
-		void SuspendContext();
 		
 	private:
 		const TCHAR* GetErrorString(SLresult Result);
diff --git a/Engine/Source/Runtime/AnimGraphRuntime/Private/AnimNodes/AnimNode_Slot.cpp b/Engine/Source/Runtime/AnimGraphRuntime/Private/AnimNodes/AnimNode_Slot.cpp
--- a/Engine/Source/Runtime/AnimGraphRuntime/Private/AnimNodes/AnimNode_Slot.cpp
+++ b/Engine/Source/Runtime/AnimGraphRuntime/Private/AnimNodes/AnimNode_Slot.cpp
@@ -34,9 +34,11 @@ void FAnimNode_Slot::Update_AnyThread(const FAnimationUpdateContext& Context)
 	// Update cache in AnimInstance.
 	Context.AnimInstanceProxy->UpdateSlotNodeWeight(SlotName, WeightData.SlotNodeWeight, Context.GetFinalBlendWeight());
 
-	if (FAnimWeight::IsRelevant(WeightData.SourceWeight))
+	const bool bUpdateSource = bAlwaysUpdateSourcePose || FAnimWeight::IsRelevant(WeightData.SourceWeight);
+	if (bUpdateSource)
 	{
-		Source.Update(Context.FractionalWeight(WeightData.SourceWeight));
+		const float SourceWeight = bAlwaysUpdateSourcePose ? FAnimWeight::GetSmallestRelevantWeight() : WeightData.SourceWeight;
+		Source.Update(Context.FractionalWeight(SourceWeight));
 	}
 }
 
diff --git a/Engine/Source/Runtime/AnimGraphRuntime/Private/PlayMontageCallbackProxy.cpp b/Engine/Source/Runtime/AnimGraphRuntime/Private/PlayMontageCallbackProxy.cpp
--- a/Engine/Source/Runtime/AnimGraphRuntime/Private/PlayMontageCallbackProxy.cpp
+++ b/Engine/Source/Runtime/AnimGraphRuntime/Private/PlayMontageCallbackProxy.cpp
@@ -8,6 +8,7 @@
 UPlayMontageCallbackProxy::UPlayMontageCallbackProxy(const FObjectInitializer& ObjectInitializer)
 	: Super(ObjectInitializer)
 	, MontageInstanceID(INDEX_NONE)
+	, bInterruptedCalledBeforeBlendingOut(false)
 {
 }
 
@@ -99,6 +100,7 @@ void UPlayMontageCallbackProxy::OnMontageBlendingOut(UAnimMontage* Montage, bool
 	if (bInterrupted)
 	{
 		OnInterrupted.Broadcast(NAME_None);
+		bInterruptedCalledBeforeBlendingOut = true;
 	}
 	else
 	{
@@ -112,6 +114,10 @@ void UPlayMontageCallbackProxy::OnMontageEnded(UAnimMontage* Montage, bool bInte
 	{
 		OnCompleted.Broadcast(NAME_None);
 	}
+	else if (!bInterruptedCalledBeforeBlendingOut)
+	{
+		OnInterrupted.Broadcast(NAME_None);
+	}
 
 	UnbindDelegates();
 }
