Allow build settings UI to show up in the static mesh editor when using the meshdescription.

diff --git a/Engine/Source/Editor/StaticMeshEditor/Private/StaticMeshEditorTools.cpp b/Engine/Source/Editor/StaticMeshEditor/Private/StaticMeshEditorTools.cpp
--- a/Engine/Source/Editor/StaticMeshEditor/Private/StaticMeshEditorTools.cpp
+++ b/Engine/Source/Editor/StaticMeshEditor/Private/StaticMeshEditorTools.cpp
@@ -3033,7 +3033,7 @@ void FLevelOfDetailSettingsLayout::AddLODLevelCategories( IDetailLayoutBuilder&
 					ReductionSettingsWidgets[LODIndex]->UpdateSettings(SrcModel.ReductionSettings);
 				}
 
-				if (SrcModel.RawMeshBulkData->IsEmpty() == false)
+				if (SrcModel.RawMeshBulkData->IsEmpty() == false || SrcModel.OriginalMeshDescription != nullptr)
 				{
 					BuildSettingsWidgets[LODIndex] = MakeShareable( new FMeshBuildSettingsLayout( AsShared() ) );
 					BuildSettingsWidgets[LODIndex]->UpdateSettings(SrcModel.BuildSettings);
