Fix reset to default not showing any visual property change for static mesh lod groups

diff --git a/Engine/Source/Editor/StaticMeshEditor/Private/StaticMeshEditor.cpp b/Engine/Source/Editor/StaticMeshEditor/Private/StaticMeshEditor.cpp
--- a/Engine/Source/Editor/StaticMeshEditor/Private/StaticMeshEditor.cpp
+++ b/Engine/Source/Editor/StaticMeshEditor/Private/StaticMeshEditor.cpp
@@ -1996,6 +1996,11 @@ void FStaticMeshEditor::NotifyPostChange( const FPropertyChangedEvent& PropertyC
 	if(StaticMesh && StaticMesh->BodySetup)
 	{
 		StaticMesh->BodySetup->CreatePhysicsMeshes();
+
+		if (GET_MEMBER_NAME_CHECKED(UStaticMesh, LODGroup) == PropertyChangedEvent.GetPropertyName())
+		{
+			RefreshTool();
+		}
 	}
 }
 
