Fix Undo issues with per-platform property overrides

diff --git a/Engine/Source/Editor/DetailCustomizations/Private/PerPlatformPropertyCustomization.cpp b/Engine/Source/Editor/DetailCustomizations/Private/PerPlatformPropertyCustomization.cpp
--- a/Engine/Source/Editor/DetailCustomizations/Private/PerPlatformPropertyCustomization.cpp
+++ b/Engine/Source/Editor/DetailCustomizations/Private/PerPlatformPropertyCustomization.cpp
@@ -17,6 +17,8 @@
 #include "ScopedTransaction.h"
 #include "IPropertyUtilities.h"
 
+#define LOCTEXT_NAMESPACE "PerPlatformPropertyCustomization"
+
 template<typename PerPlatformType>
 void FPerPlatformPropertyCustomization<PerPlatformType>::CustomizeHeader(TSharedRef<IPropertyHandle> StructPropertyHandle, FDetailWidgetRow& HeaderRow, IPropertyTypeCustomizationUtils& StructCustomizationUtils)
 {
@@ -108,6 +110,8 @@ float FPerPlatformPropertyCustomization<PerPlatformType>::CalcDesiredWidth(TShar
 template<typename PerPlatformType>
 bool FPerPlatformPropertyCustomization<PerPlatformType>::AddPlatformOverride(FName PlatformGroupName, TSharedRef<IPropertyHandle> StructPropertyHandle)
 {
+	FScopedTransaction Transaction(LOCTEXT("AddPlatformOverride", "Add Platform Override"));
+
 	TSharedPtr<IPropertyHandle>	PerPlatformProperty = StructPropertyHandle->GetChildHandle(FName("PerPlatform"));
 	TSharedPtr<IPropertyHandle>	DefaultProperty = StructPropertyHandle->GetChildHandle(FName("Default"));
 	if (PerPlatformProperty.IsValid() && DefaultProperty.IsValid())
@@ -156,6 +160,8 @@ bool FPerPlatformPropertyCustomization<PerPlatformType>::AddPlatformOverride(FNa
 template<typename PerPlatformType>
 bool FPerPlatformPropertyCustomization<PerPlatformType>::RemovePlatformOverride(FName PlatformGroupName, TSharedRef<IPropertyHandle> StructPropertyHandle)
 {
+	FScopedTransaction Transaction(LOCTEXT("RemovePlatformOverride", "Remove Platform Override"));
+
 	TSharedPtr<IPropertyHandle>	PerPlatformProperty = StructPropertyHandle->GetChildHandle(FName("PerPlatform"));
 	if (PerPlatformProperty.IsValid())
 	{
@@ -232,3 +238,4 @@ TSharedRef<IPropertyTypeCustomization> FPerPlatformPropertyCustomization<PerPlat
 template class FPerPlatformPropertyCustomization<FPerPlatformInt>;
 template class FPerPlatformPropertyCustomization<FPerPlatformFloat>;
 
+#undef LOCTEXT_NAMESPACE
\ No newline at end of file
diff --git a/Engine/Source/Editor/StaticMeshEditor/Private/StaticMeshEditorTools.cpp b/Engine/Source/Editor/StaticMeshEditor/Private/StaticMeshEditorTools.cpp
--- a/Engine/Source/Editor/StaticMeshEditor/Private/StaticMeshEditorTools.cpp
+++ b/Engine/Source/Editor/StaticMeshEditor/Private/StaticMeshEditorTools.cpp
@@ -3312,11 +3312,13 @@ float FLevelOfDetailSettingsLayout::GetScreenSizeWidgetWidth(int32 LODIndex) con
 
 bool FLevelOfDetailSettingsLayout::AddLODScreenSizePlatformOverride(FName PlatformGroupName, int32 LODIndex)
 {
+	FScopedTransaction Transaction(LOCTEXT("AddLODScreenSizePlatformOverride", "Add LOD Screen Size Platform Override"));
 	UStaticMesh* StaticMesh = StaticMeshEditor.GetStaticMesh();
 	if (LODScreenSizes[LODIndex].PerPlatform.Find(PlatformGroupName) == nullptr)
 	{
 		if(!StaticMesh->bAutoComputeLODScreenSize && StaticMesh->SourceModels.IsValidIndex(LODIndex))
 		{
+			StaticMesh->Modify();
 			float Value = StaticMesh->SourceModels[LODIndex].ScreenSize.Default;
 			StaticMesh->SourceModels[LODIndex].ScreenSize.PerPlatform.Add(PlatformGroupName, Value);
 			OnLODScreenSizeChanged(Value, PlatformGroupName, LODIndex);
@@ -3328,9 +3330,11 @@ bool FLevelOfDetailSettingsLayout::AddLODScreenSizePlatformOverride(FName Platfo
 
 bool FLevelOfDetailSettingsLayout::RemoveLODScreenSizePlatformOverride(FName PlatformGroupName, int32 LODIndex)
 {
+	FScopedTransaction Transaction(LOCTEXT("RemoveLODScreenSizePlatformOverride", "Remove LOD Screen Size Platform Override"));
 	UStaticMesh* StaticMesh = StaticMeshEditor.GetStaticMesh();
 	if (!StaticMesh->bAutoComputeLODScreenSize && StaticMesh->SourceModels.IsValidIndex(LODIndex))
 	{
+		StaticMesh->Modify();
 		if (StaticMesh->SourceModels[LODIndex].ScreenSize.PerPlatform.Remove(PlatformGroupName) != 0)
 		{
 			OnLODScreenSizeChanged(StaticMesh->SourceModels[LODIndex].ScreenSize.Default, PlatformGroupName, LODIndex);
@@ -3696,8 +3700,10 @@ TSharedRef<SWidget> FLevelOfDetailSettingsLayout::GetMinLODWidget(FName Platform
 
 bool FLevelOfDetailSettingsLayout::AddMinLODPlatformOverride(FName PlatformGroupName)
 {
+	FScopedTransaction Transaction(LOCTEXT("AddMinLODPlatformOverride", "Add Min LOD Platform Override"));
 	UStaticMesh* StaticMesh = StaticMeshEditor.GetStaticMesh();
 	check(StaticMesh);
+	StaticMesh->Modify();
 	if (StaticMesh->MinLOD.PerPlatform.Find(PlatformGroupName) == nullptr)
 	{
 		float Value = StaticMesh->MinLOD.Default;
@@ -3710,8 +3716,10 @@ bool FLevelOfDetailSettingsLayout::AddMinLODPlatformOverride(FName PlatformGroup
 
 bool FLevelOfDetailSettingsLayout::RemoveMinLODPlatformOverride(FName PlatformGroupName)
 {
+	FScopedTransaction Transaction(LOCTEXT("RemoveMinLODPlatformOverride", "Remove Min LOD Platform Override"));
 	UStaticMesh* StaticMesh = StaticMeshEditor.GetStaticMesh();
 	check(StaticMesh);
+	StaticMesh->Modify();
 	if (StaticMesh->MinLOD.PerPlatform.Remove(PlatformGroupName) != 0)
 	{
 		OnMinLODChanged(StaticMesh->MinLOD.Default, PlatformGroupName);
