Fixed "Mesh Editing" toolbar in Static Mesh Editor showing as "Secondary Toolbar" when closing and re-opening the tab.

diff --git a/Engine/Source/Editor/StaticMeshEditor/Private/StaticMeshEditor.cpp b/Engine/Source/Editor/StaticMeshEditor/Private/StaticMeshEditor.cpp
--- a/Engine/Source/Editor/StaticMeshEditor/Private/StaticMeshEditor.cpp
+++ b/Engine/Source/Editor/StaticMeshEditor/Private/StaticMeshEditor.cpp
@@ -504,8 +504,10 @@ TSharedRef<SDockTab> FStaticMeshEditor::SpawnTab_SecondaryToolbar( const FSpawnT
 {
 	check( Args.GetTabId() == SecondaryToolbarTabId );
 
+	FText TabLabel = !SecondaryToolbarDisplayName.IsEmpty() ? SecondaryToolbarDisplayName : LOCTEXT("SecondaryToolbar_TabTitle", "Secondary Toolbar");
+
 	TSharedRef<SDockTab> SpawnedTab = SNew(SDockTab)
-		.Label( LOCTEXT("SecondaryToolbar_TabTitle", "Secondary Toolbar") )
+		.Label( TabLabel )
 		.Icon( FEditorStyle::GetBrush("LevelEditor.Tabs.Toolbar") )
 		.ShouldAutosize( true )
 		[
