Copied 4356072 Anousack.Kitisa
[StaticMeshEditor]:
* Moved the UV channel removal functionality from the LOD details panel to the UV drop-down menu.

diff --git a/Engine/Source/Editor/StaticMeshEditor/Private/StaticMeshEditor.cpp b/Engine/Source/Editor/StaticMeshEditor/Private/StaticMeshEditor.cpp
--- a/Engine/Source/Editor/StaticMeshEditor/Private/StaticMeshEditor.cpp
+++ b/Engine/Source/Editor/StaticMeshEditor/Private/StaticMeshEditor.cpp
@@ -1180,6 +1180,24 @@ TSharedRef<SWidget> FStaticMeshEditor::GenerateUVChannelComboList()
 		);
 	}
 
+	// Add UV editing functions
+	{
+		MenuBuilder.AddMenuSeparator();
+
+		FUIAction MenuAction;
+		MenuAction.ExecuteAction.BindSP(this, &FStaticMeshEditor::RemoveCurrentUVChannel);
+		MenuAction.CanExecuteAction.BindSP(this, &FStaticMeshEditor::CanRemoveUVChannel);
+
+		MenuBuilder.AddMenuEntry(
+			LOCTEXT("Remove_UVChannel", "Remove Selected"),
+			LOCTEXT("Remove_UVChannel_ToolTip", "Remove currently selected UV channel from the static mesh"),
+			FSlateIcon(),
+			MenuAction,
+			NAME_None,
+			EUserInterfaceActionType::Button
+		);
+	}
+
 	MenuBuilder.EndSection();
 
 	return MenuBuilder.MakeWidget();
@@ -2091,5 +2109,51 @@ TStatId FStaticMeshEditor::GetStatId() const
 	RETURN_QUICK_DECLARE_CYCLE_STAT(FStaticMeshEditor, STATGROUP_TaskGraphTasks);
 }
 
+bool FStaticMeshEditor::CanRemoveUVChannel()
+{
+	// Can remove UV channel if there's one that is currently being selected and displayed
+	return Viewport->GetViewportClient().IsDrawUVOverlayChecked();
+}
+
+void FStaticMeshEditor::RemoveCurrentUVChannel()
+{
+	if (!StaticMesh)
+	{
+		return;
+	}
+
+	int32 UVChannelIndex = GetCurrentUVChannel();
+	int32 LODIndex = GetCurrentLODIndex();
+
+	FText RemoveUVChannelText = FText::Format(LOCTEXT("ConfirmRemoveUVChannel", "Please confirm removal of UV Channel {0} from LOD {1} of {2}?"), UVChannelIndex, LODIndex, FText::FromString(StaticMesh->GetName()));
+	if (FMessageDialog::Open(EAppMsgType::YesNo, RemoveUVChannelText) == EAppReturnType::Yes)
+	{
+		FMeshBuildSettings& LODBuildSettings = StaticMesh->SourceModels[LODIndex].BuildSettings;
+
+		if (LODBuildSettings.bGenerateLightmapUVs)
+		{
+			FText LightmapText;
+			if (UVChannelIndex == LODBuildSettings.SrcLightmapIndex)
+			{
+				LightmapText = FText::Format(LOCTEXT("ConfirmDisableSourceLightmap", "UV Channel {0} is currently used as source for lightmap UVs. Please change the \"Source Lightmap Index\" value or disable \"Generate Lightmap UVs\" in the Build Settings."), UVChannelIndex);
+			}
+			else if (UVChannelIndex == LODBuildSettings.DstLightmapIndex)
+			{
+				LightmapText = FText::Format(LOCTEXT("ConfirmDisableDestLightmap", "UV Channel {0} is currently used as destination for lightmap UVs. Please change the \"Destination Lightmap Index\" value or disable \"Generate Lightmap UVs\" in the Build Settings."), UVChannelIndex);
+			}
+
+			if (!LightmapText.IsEmpty())
+			{
+				FMessageDialog::Open(EAppMsgType::Ok, LightmapText);
+				return;
+			}
+		}
+
+		if (StaticMesh->RemoveUVChannel(LODIndex, UVChannelIndex))
+		{
+			RefreshTool();
+		}
+	}
+}
 
 #undef LOCTEXT_NAMESPACE
diff --git a/Engine/Source/Editor/StaticMeshEditor/Private/StaticMeshEditor.h b/Engine/Source/Editor/StaticMeshEditor/Private/StaticMeshEditor.h
--- a/Engine/Source/Editor/StaticMeshEditor/Private/StaticMeshEditor.h
+++ b/Engine/Source/Editor/StaticMeshEditor/Private/StaticMeshEditor.h
@@ -336,6 +336,10 @@ private:
 
 	ECheckBoxState GetUVChannelCheckState(int32 TestUVChannel) const;
 
+	/** Callbacks for UV edition */
+	bool CanRemoveUVChannel();
+	void RemoveCurrentUVChannel();
+
 private:
 	/** List of open tool panels; used to ensure only one exists at any one time */
 	TMap< FName, TWeakPtr<class SDockableTab> > SpawnedToolPanels;
diff --git a/Engine/Source/Editor/StaticMeshEditor/Private/StaticMeshEditorTools.cpp b/Engine/Source/Editor/StaticMeshEditor/Private/StaticMeshEditorTools.cpp
--- a/Engine/Source/Editor/StaticMeshEditor/Private/StaticMeshEditorTools.cpp
+++ b/Engine/Source/Editor/StaticMeshEditor/Private/StaticMeshEditorTools.cpp
@@ -39,7 +39,6 @@
 #include "Widgets/Input/SFilePathPicker.h"
 #include "EditorDirectories.h"
 #include "EditorFramework/AssetImportData.h"
-#include "GenericPlatform/GenericPlatformMath.h"
 
 const uint32 MaxHullCount = 64;
 const uint32 MinHullCount = 2;
@@ -3179,12 +3178,6 @@ void FLevelOfDetailSettingsLayout::AddLODLevelCategories( IDetailLayoutBuilder&
 				LODCategory.AddCustomBuilder( ReductionSettingsWidgets[LODIndex].ToSharedRef() );
 			}
 
-			UVChannelsWidgets[LODIndex] = MakeShareable(new FUVChannelsLayout(StaticMeshEditor, LODIndex));
-			if (UVChannelsWidgets[LODIndex].IsValid())
-			{
-				LODCategory.AddCustomBuilder(UVChannelsWidgets[LODIndex].ToSharedRef());
-			}
-
 			if (LODIndex != 0)
 			{
 				LODCategory.AddCustomRow( LOCTEXT("RemoveLOD", "Remove LOD") )
@@ -3974,201 +3967,5 @@ FText FLevelOfDetailSettingsLayout::GetCurrentLodTooltip() const
 	return FText::GetEmpty();
 }
 
-FUVChannelsLayout::FUVChannelsLayout(IStaticMeshEditor& InStaticMeshEditor, int32 InLODIndex)
-	: StaticMeshEditor(InStaticMeshEditor)
-	, LODIndex(InLODIndex)
-	, UVChannelIndex(0)
-{
-	UpdateNumUVChannels();
-}
-
-FUVChannelsLayout::~FUVChannelsLayout()
-{
-}
-
-void FUVChannelsLayout::UpdateNumUVChannels()
-{
-	if (StaticMeshEditor.GetStaticMesh())
-	{
-		NumUVChannels = StaticMeshEditor.GetStaticMesh()->GetNumUVChannels(LODIndex);
-	}
-}
-
-void FUVChannelsLayout::GenerateHeaderRowContent(FDetailWidgetRow& NodeRow)
-{
-	NodeRow.NameContent()
-	[
-		SNew(STextBlock)
-		.Text(LOCTEXT("UVChannels", "UV Channels"))
-		.Font(IDetailLayoutBuilder::GetDetailFont())
-	];
-}
-
-void FUVChannelsLayout::GenerateChildContent(IDetailChildrenBuilder& ChildrenBuilder)
-{
-	{
-		ChildrenBuilder.AddCustomRow(LOCTEXT("NumUVChannels", "Number of UV Channels"))
-		.NameContent()
-		[
-			SNew(STextBlock)
-			.Font(IDetailLayoutBuilder::GetDetailFont())
-			.Text(LOCTEXT("NumUVChannels", "Number of UV Channels"))
-		]
-		.ValueContent()
-		[
-			SNew(STextBlock)
-			.Font(IDetailLayoutBuilder::GetDetailFont())
-			.Text(this, &FUVChannelsLayout::GetNumUVChannelsAsText)
-		];
-	}
-
-	{
-		ChildrenBuilder.AddCustomRow(LOCTEXT("UVChannelIndex", "UV Channel Index"))
-		.NameContent()
-		[
-			SNew(STextBlock)
-			.Font(IDetailLayoutBuilder::GetDetailFont())
-			.Text(LOCTEXT("UVChannelIndex", "UV Channel Index"))
-		]
-		.ValueContent()
-		[
-			SNew(SHorizontalBox)
-			+ SHorizontalBox::Slot()
-			.VAlign(VAlign_Center)
-			.AutoWidth()
-			[
-				SNew(SSpinBox<int32>)
-				.Font(IDetailLayoutBuilder::GetDetailFont())
-				.MinDesiredWidth(100.f)
-				.MinValue(0)
-				.MaxValue(this, &FUVChannelsLayout::GetMaxUVChannelIndex)
-				.Value(this, &FUVChannelsLayout::GetUVChannelIndex)
-				.OnValueChanged(this, &FUVChannelsLayout::OnUVChannelIndexChanged)
-			]
-
-			+ SHorizontalBox::Slot()
-			.VAlign(VAlign_Center)
-			.AutoWidth()
-			.Padding(4.0f, 0.0f, 0.0f, 0.0f)
-			[
-				SNew(SButton)
-				.ToolTipText(LOCTEXT("InsertUVChannel_Tooltip", "Inserts a UV Channel at index"))
-				.OnClicked(this, &FUVChannelsLayout::OnInsertUVChannel)
-				.IsEnabled(this, &FUVChannelsLayout::CanAddUVChannel)
-				[
-					SNew(STextBlock)
-					.Text(LOCTEXT("InsertUVChannel", "Insert"))
-					.Font(IDetailLayoutBuilder::GetDetailFont())
-				]
-			]
-
-			+ SHorizontalBox::Slot()
-			.VAlign(VAlign_Center)
-			.AutoWidth()
-			.Padding(4.0f, 0.0f, 0.0f, 0.0f)
-			[
-				SNew(SButton)
-				.ToolTipText(LOCTEXT("RemoveUVChannel_Tooltip", "Removes the UV Channel at index"))
-				.OnClicked(this, &FUVChannelsLayout::OnRemoveUVChannel)
-				.IsEnabled(this, &FUVChannelsLayout::CanRemoveUVChannel)
-				[
-					SNew(STextBlock)
-					.Text(LOCTEXT("RemoveUVChannel", "Remove"))
-					.Font(IDetailLayoutBuilder::GetDetailFont())
-				]
-			]
-		];
-	}
-}
-
-FText FUVChannelsLayout::GetNumUVChannelsAsText() const
-{
-	return FText::AsNumber(NumUVChannels);
-}
-
-TOptional<int32> FUVChannelsLayout::GetMaxUVChannelIndex() const
-{
-	return FGenericPlatformMath::Min(FGenericPlatformMath::Max(NumUVChannels, 0), MAX_MESH_TEXTURE_COORDS - 1);
-}
-
-int32 FUVChannelsLayout::GetUVChannelIndex() const
-{
-	return UVChannelIndex;
-}
-
-void FUVChannelsLayout::OnUVChannelIndexChanged(int32 NewValue)
-{
-	UVChannelIndex = NewValue;
-}
-
-FReply FUVChannelsLayout::OnInsertUVChannel()
-{
-	if (StaticMeshEditor.GetStaticMesh() && StaticMeshEditor.GetStaticMesh()->InsertUVChannel(LODIndex, UVChannelIndex))
-	{
-		StaticMeshEditor.RefreshTool();
-		return FReply::Handled();
-	}
-
-	return FReply::Unhandled();
-}
-
-FReply FUVChannelsLayout::OnRemoveUVChannel()
-{
-	UStaticMesh* StaticMesh = StaticMeshEditor.GetStaticMesh();
-	if (!StaticMesh)
-	{
-		return FReply::Unhandled();
-	}
-
-	FText RemoveUVChannelText = FText::Format(LOCTEXT("ConfirmRemoveUVChannel", "Are you sure you want to remove UV Channel {0} from LOD {1} of {2}?"), UVChannelIndex, LODIndex, FText::FromString(StaticMesh->GetName()));
-	if (FMessageDialog::Open(EAppMsgType::YesNo, RemoveUVChannelText) == EAppReturnType::Yes)
-	{
-		FMeshBuildSettings& LODBuildSettings = StaticMesh->SourceModels[LODIndex].BuildSettings;
-
-		if (LODBuildSettings.bGenerateLightmapUVs)
-		{
-			FText LightmapText;
-			if (UVChannelIndex == LODBuildSettings.SrcLightmapIndex)
-			{
-				LightmapText = FText::Format(LOCTEXT("ConfirmDisableSourceLightmap", "To remove lightmap source UV at index {0}, \"Generate Lightmap UVs\" must be disabled in the Build Settings. Do you want to disable it?"), UVChannelIndex);
-			}
-			else if (UVChannelIndex == LODBuildSettings.DstLightmapIndex)
-			{
-				LightmapText = FText::Format(LOCTEXT("ConfirmDisableDestLightmap", "To remove lightmap destination UV at index {0}, \"Generate Lightmap UVs\" must be disabled in the Build Settings. Do you want to disable it?"), UVChannelIndex);
-			}
-
-			if (!LightmapText.IsEmpty())
-			{
-				if (FMessageDialog::Open(EAppMsgType::YesNo, LightmapText) == EAppReturnType::Yes)
-				{
-					LODBuildSettings.bGenerateLightmapUVs = false;
-				}
-				else
-				{
-					return FReply::Unhandled();
-				}
-			}
-		}
-
-		if (StaticMesh->RemoveUVChannel(LODIndex, UVChannelIndex))
-		{
-			StaticMeshEditor.RefreshTool();
-			return FReply::Handled();
-		}
-	}
-
-	return FReply::Unhandled();
-}
-
-bool FUVChannelsLayout::CanAddUVChannel() const
-{
-	return NumUVChannels < MAX_MESH_TEXTURE_COORDS;
-}
-
-bool FUVChannelsLayout::CanRemoveUVChannel() const
-{
-	// Some mesh operations assume there's a channel 0
-	return NumUVChannels > 1 && UVChannelIndex < NumUVChannels;
-}
 
 #undef LOCTEXT_NAMESPACE
diff --git a/Engine/Source/Editor/StaticMeshEditor/Private/StaticMeshEditorTools.h b/Engine/Source/Editor/StaticMeshEditor/Private/StaticMeshEditorTools.h
--- a/Engine/Source/Editor/StaticMeshEditor/Private/StaticMeshEditorTools.h
+++ b/Engine/Source/Editor/StaticMeshEditor/Private/StaticMeshEditorTools.h
@@ -441,40 +441,6 @@ private:
 	bool bDeleteWarningConsumed;
 };
 
-class FUVChannelsLayout : public IDetailCustomNodeBuilder, public TSharedFromThis<FUVChannelsLayout>
-{
-public:
-	FUVChannelsLayout(IStaticMeshEditor& InStaticMeshEditor, int32 InLODIndex);
-	virtual ~FUVChannelsLayout();
-
-private:
-	/** IDetailCustomNodeBuilder Interface*/
-	virtual void SetOnRebuildChildren(FSimpleDelegate InOnRegenerateChildren) override {}
-	virtual void GenerateHeaderRowContent(FDetailWidgetRow& NodeRow) override;
-	virtual void GenerateChildContent(IDetailChildrenBuilder& ChildrenBuilder) override;
-	virtual void Tick(float DeltaTime) override {}
-	virtual bool RequiresTick() const override { return false; }
-	virtual FName GetName() const override { static FName UVChannelsName("UVChannels"); return UVChannelsName; }
-	virtual bool InitiallyCollapsed() const override { return true; }
-
-	void UpdateNumUVChannels();
-	FText GetNumUVChannelsAsText() const;
-	TOptional<int32> GetMaxUVChannelIndex() const;
-	int32 GetUVChannelIndex() const;
-	void OnUVChannelIndexChanged(int32 NewValue);
-
-	FReply OnInsertUVChannel();
-	FReply OnRemoveUVChannel();
-	bool CanAddUVChannel() const;
-	bool CanRemoveUVChannel() const;
-
-private:
-	IStaticMeshEditor& StaticMeshEditor;
-	int32 LODIndex;
-	int32 NumUVChannels;
-	int32 UVChannelIndex;
-};
-
 /** 
  * Window for adding and removing LOD.
  */
@@ -573,7 +539,6 @@ private:
 	TSharedPtr<FMeshReductionSettingsLayout> ReductionSettingsWidgets[MAX_STATIC_MESH_LODS];
 	TSharedPtr<FMeshBuildSettingsLayout> BuildSettingsWidgets[MAX_STATIC_MESH_LODS];
 	TSharedPtr<FMeshSectionSettingsLayout> SectionSettingsWidgets[MAX_STATIC_MESH_LODS];
-	TSharedPtr<FUVChannelsLayout> UVChannelsWidgets[MAX_STATIC_MESH_LODS];
 
 	TSharedPtr<FMeshMaterialsLayout> MaterialsLayoutWidget;
 
