Copied 4399602 Anousack.Kitisa
[StaticMeshEditor]:
* Added a secondary toolbar tab that is hidden by default (and when it's empty) but visible when it is populated by extenders.

diff --git a/Engine/Plugins/Editor/MeshEditor/Source/MeshEditor/MeshEditorStyle.cpp b/Engine/Plugins/Editor/MeshEditor/Source/MeshEditor/MeshEditorStyle.cpp
--- a/Engine/Plugins/Editor/MeshEditor/Source/MeshEditor/MeshEditorStyle.cpp
+++ b/Engine/Plugins/Editor/MeshEditor/Source/MeshEditor/MeshEditorStyle.cpp
@@ -40,9 +40,9 @@ void FMeshEditorStyle::Initialize()
 
 	// Icons for the mode panel tabs
 	StyleSet->Set( "LevelEditor.MeshEditorMode", new IMAGE_PLUGIN_BRUSH( "Icons/MeshEditorMode_40px", Icon40x40 ) );
-	StyleSet->Set( "LevelEditor.MeshEditorMode.Small", new IMAGE_PLUGIN_BRUSH( "Icons/MeshEditorMode_40px", Icon40x40 ) );
+	StyleSet->Set( "LevelEditor.MeshEditorMode.Small", new IMAGE_PLUGIN_BRUSH( "Icons/MeshEditorMode_40px", Icon20x20) );
 	StyleSet->Set( "LevelEditor.MeshEditorMode.Selected", new IMAGE_PLUGIN_BRUSH( "Icons/MeshEditorMode_40px", Icon40x40 ) );
-	StyleSet->Set( "LevelEditor.MeshEditorMode.Selected.Small", new IMAGE_PLUGIN_BRUSH( "Icons/MeshEditorMode_40px", Icon40x40 ) );
+	StyleSet->Set( "LevelEditor.MeshEditorMode.Selected.Small", new IMAGE_PLUGIN_BRUSH( "Icons/MeshEditorMode_40px", Icon20x20) );
 
 	StyleSet->Set( "EditingMode.GroupName.Text", FTextBlockStyle()
 		.SetFont( TTF_CORE_FONT( "Fonts/Roboto-Bold", 12 ) )
@@ -113,44 +113,44 @@ void FMeshEditorStyle::Initialize()
 	StyleSet->Set("MeshEditorMode.VertexEditMode", new IMAGE_PLUGIN_BRUSH("Icons/T_Radial_Mesh_Vertex", Icon512x512));
 
 	StyleSet->Set( "MeshEditorAnyElement.DeleteMeshElement", new IMAGE_PLUGIN_BRUSH( "Icons/DeleteMeshElement", Icon40x40 ) );
-	StyleSet->Set( "MeshEditorAnyElement.DeleteMeshElement.Small", new IMAGE_PLUGIN_BRUSH( "Icons/DeleteMeshElement", Icon40x40 ) );
+	StyleSet->Set( "MeshEditorAnyElement.DeleteMeshElement.Small", new IMAGE_PLUGIN_BRUSH( "Icons/DeleteMeshElement", Icon20x20) );
 	StyleSet->Set( "MeshEditorAnyElement.DeleteMeshElement.Selected", new IMAGE_PLUGIN_BRUSH( "Icons/DeleteMeshElement", Icon40x40 ) );
-	StyleSet->Set( "MeshEditorAnyElement.DeleteMeshElement.Selected.Small", new IMAGE_PLUGIN_BRUSH( "Icons/DeleteMeshElement", Icon40x40 ) );
+	StyleSet->Set( "MeshEditorAnyElement.DeleteMeshElement.Selected.Small", new IMAGE_PLUGIN_BRUSH( "Icons/DeleteMeshElement", Icon20x20) );
 
 	StyleSet->Set( "MeshEditorPolygon.FlipPolygon", new IMAGE_PLUGIN_BRUSH( "Icons/FlipPolygon", Icon40x40 ) );
-	StyleSet->Set( "MeshEditorPolygon.FlipPolygon.Small", new IMAGE_PLUGIN_BRUSH( "Icons/FlipPolygon", Icon40x40 ) );
+	StyleSet->Set( "MeshEditorPolygon.FlipPolygon.Small", new IMAGE_PLUGIN_BRUSH( "Icons/FlipPolygon", Icon20x20) );
 	StyleSet->Set( "MeshEditorPolygon.FlipPolygon.Selected", new IMAGE_PLUGIN_BRUSH( "Icons/FlipPolygon", Icon40x40 ) );
-	StyleSet->Set( "MeshEditorPolygon.FlipPolygon.Selected.Small", new IMAGE_PLUGIN_BRUSH( "Icons/FlipPolygon", Icon40x40 ) );
+	StyleSet->Set( "MeshEditorPolygon.FlipPolygon.Selected.Small", new IMAGE_PLUGIN_BRUSH( "Icons/FlipPolygon", Icon20x20) );
 
 	StyleSet->Set( "MeshEditorPolygon.AssignMaterial", new IMAGE_PLUGIN_BRUSH( "Icons/AssignMaterial", Icon40x40 ) );
-	StyleSet->Set( "MeshEditorPolygon.AssignMaterial.Small", new IMAGE_PLUGIN_BRUSH( "Icons/AssignMaterial", Icon40x40 ) );
+	StyleSet->Set( "MeshEditorPolygon.AssignMaterial.Small", new IMAGE_PLUGIN_BRUSH( "Icons/AssignMaterial", Icon20x20) );
 	StyleSet->Set( "MeshEditorPolygon.AssignMaterial.Selected", new IMAGE_PLUGIN_BRUSH( "Icons/AssignMaterial", Icon40x40 ) );
-	StyleSet->Set( "MeshEditorPolygon.AssignMaterial.Selected.Small", new IMAGE_PLUGIN_BRUSH( "Icons/AssignMaterial", Icon40x40 ) );
+	StyleSet->Set( "MeshEditorPolygon.AssignMaterial.Selected.Small", new IMAGE_PLUGIN_BRUSH( "Icons/AssignMaterial", Icon20x20) );
 
 	StyleSet->Set( "MeshEditorPolygon.UnifyNormals", new IMAGE_PLUGIN_BRUSH( "Icons/UnifyNormals", Icon40x40 ) );
-	StyleSet->Set( "MeshEditorPolygon.UnifyNormals.Small", new IMAGE_PLUGIN_BRUSH( "Icons/UnifyNormals", Icon40x40 ) );
+	StyleSet->Set( "MeshEditorPolygon.UnifyNormals.Small", new IMAGE_PLUGIN_BRUSH( "Icons/UnifyNormals", Icon20x20) );
 	StyleSet->Set( "MeshEditorPolygon.UnifyNormals.Selected", new IMAGE_PLUGIN_BRUSH( "Icons/UnifyNormals", Icon40x40 ) );
-	StyleSet->Set( "MeshEditorPolygon.UnifyNormals.Selected.Small", new IMAGE_PLUGIN_BRUSH( "Icons/UnifyNormals", Icon40x40 ) );
+	StyleSet->Set( "MeshEditorPolygon.UnifyNormals.Selected.Small", new IMAGE_PLUGIN_BRUSH( "Icons/UnifyNormals", Icon20x20) );
 
 	StyleSet->Set( "MeshEditorSelectionModifiers.PolygonsByGroup", new IMAGE_PLUGIN_BRUSH( "Icons/PolygonsByGroup", Icon40x40 ) );
-	StyleSet->Set( "MeshEditorSelectionModifiers.PolygonsByGroup.Small", new IMAGE_PLUGIN_BRUSH( "Icons/PolygonsByGroup", Icon40x40 ) );
+	StyleSet->Set( "MeshEditorSelectionModifiers.PolygonsByGroup.Small", new IMAGE_PLUGIN_BRUSH( "Icons/PolygonsByGroup", Icon20x20) );
 	StyleSet->Set( "MeshEditorSelectionModifiers.PolygonsByGroup.Selected", new IMAGE_PLUGIN_BRUSH( "Icons/PolygonsByGroup", Icon40x40 ) );
-	StyleSet->Set( "MeshEditorSelectionModifiers.PolygonsByGroup.Selected.Small", new IMAGE_PLUGIN_BRUSH( "Icons/PolygonsByGroup", Icon40x40 ) );
+	StyleSet->Set( "MeshEditorSelectionModifiers.PolygonsByGroup.Selected.Small", new IMAGE_PLUGIN_BRUSH( "Icons/PolygonsByGroup", Icon20x20) );
 
 	StyleSet->Set( "MeshEditorSelectionModifiers.SingleElement", new IMAGE_PLUGIN_BRUSH( "Icons/SingleElement", Icon40x40 ) );
-	StyleSet->Set( "MeshEditorSelectionModifiers.SingleElement.Small", new IMAGE_PLUGIN_BRUSH( "Icons/SingleElement", Icon40x40 ) );
+	StyleSet->Set( "MeshEditorSelectionModifiers.SingleElement.Small", new IMAGE_PLUGIN_BRUSH( "Icons/SingleElement", Icon20x20) );
 	StyleSet->Set( "MeshEditorSelectionModifiers.SingleElement.Selected", new IMAGE_PLUGIN_BRUSH( "Icons/SingleElement", Icon40x40 ) );
-	StyleSet->Set( "MeshEditorSelectionModifiers.SingleElement.Selected.Small", new IMAGE_PLUGIN_BRUSH( "Icons/SingleElement", Icon40x40 ) );
+	StyleSet->Set( "MeshEditorSelectionModifiers.SingleElement.Selected.Small", new IMAGE_PLUGIN_BRUSH( "Icons/SingleElement", Icon20x20) );
 
 	StyleSet->Set( "MeshEditorSelectionModifiers.PolygonsByConnectivity", new IMAGE_PLUGIN_BRUSH( "Icons/PolygonsByConnectivity", Icon40x40 ) );
-	StyleSet->Set( "MeshEditorSelectionModifiers.PolygonsByConnectivity.Small", new IMAGE_PLUGIN_BRUSH( "Icons/PolygonsByConnectivity", Icon40x40 ) );
+	StyleSet->Set( "MeshEditorSelectionModifiers.PolygonsByConnectivity.Small", new IMAGE_PLUGIN_BRUSH( "Icons/PolygonsByConnectivity", Icon20x20) );
 	StyleSet->Set( "MeshEditorSelectionModifiers.PolygonsByConnectivity.Selected", new IMAGE_PLUGIN_BRUSH( "Icons/PolygonsByConnectivity", Icon40x40 ) );
-	StyleSet->Set( "MeshEditorSelectionModifiers.PolygonsByConnectivity.Selected.Small", new IMAGE_PLUGIN_BRUSH( "Icons/PolygonsByConnectivity", Icon40x40 ) );
+	StyleSet->Set( "MeshEditorSelectionModifiers.PolygonsByConnectivity.Selected.Small", new IMAGE_PLUGIN_BRUSH( "Icons/PolygonsByConnectivity", Icon20x20) );
 
 	StyleSet->Set("MeshEditorSelectionModifiers.PolygonsBySmoothingGroup", new IMAGE_PLUGIN_BRUSH("Icons/PolygonsBySmoothingGroup", Icon40x40));
-	StyleSet->Set("MeshEditorSelectionModifiers.PolygonsBySmoothingGroup.Small", new IMAGE_PLUGIN_BRUSH("Icons/PolygonsBySmoothingGroup", Icon40x40));
+	StyleSet->Set("MeshEditorSelectionModifiers.PolygonsBySmoothingGroup.Small", new IMAGE_PLUGIN_BRUSH("Icons/PolygonsBySmoothingGroup", Icon20x20));
 	StyleSet->Set("MeshEditorSelectionModifiers.PolygonsBySmoothingGroup.Selected", new IMAGE_PLUGIN_BRUSH("Icons/PolygonsBySmoothingGroup", Icon40x40));
-	StyleSet->Set("MeshEditorSelectionModifiers.PolygonsBySmoothingGroup.Selected.Small", new IMAGE_PLUGIN_BRUSH("Icons/PolygonsBySmoothingGroup", Icon40x40));
+	StyleSet->Set("MeshEditorSelectionModifiers.PolygonsBySmoothingGroup.Selected.Small", new IMAGE_PLUGIN_BRUSH("Icons/PolygonsBySmoothingGroup", Icon20x20));
 
 	FSlateStyleRegistry::RegisterSlateStyle( *StyleSet.Get() );
 }
diff --git a/Engine/Source/Editor/StaticMeshEditor/Private/StaticMeshEditor.cpp b/Engine/Source/Editor/StaticMeshEditor/Private/StaticMeshEditor.cpp
--- a/Engine/Source/Editor/StaticMeshEditor/Private/StaticMeshEditor.cpp
+++ b/Engine/Source/Editor/StaticMeshEditor/Private/StaticMeshEditor.cpp
@@ -65,7 +65,8 @@ const FName FStaticMeshEditor::ViewportTabId( TEXT( "StaticMeshEditor_Viewport"
 const FName FStaticMeshEditor::PropertiesTabId( TEXT( "StaticMeshEditor_Properties" ) );
 const FName FStaticMeshEditor::SocketManagerTabId( TEXT( "StaticMeshEditor_SocketManager" ) );
 const FName FStaticMeshEditor::CollisionTabId( TEXT( "StaticMeshEditor_Collision" ) );
-const FName FStaticMeshEditor::PreviewSceneSettingsTabId(TEXT("StaticMeshEditor_PreviewScene"));
+const FName FStaticMeshEditor::PreviewSceneSettingsTabId( TEXT ("StaticMeshEditor_PreviewScene" ) );
+const FName FStaticMeshEditor::SecondaryToolbarTabId( TEXT( "StaticMeshEditor_SecondaryToolbar" ) );
 
 void FStaticMeshEditor::RegisterTabSpawners(const TSharedRef<class FTabManager>& InTabManager)
 {
@@ -94,10 +95,19 @@ void FStaticMeshEditor::RegisterTabSpawners(const TSharedRef<class FTabManager>&
 		.SetGroup(WorkspaceMenuCategoryRef)
 		.SetIcon(FSlateIcon(FEditorStyle::GetStyleSetName(), "StaticMeshEditor.Tabs.ConvexDecomposition"));
 
-	InTabManager->RegisterTabSpawner(PreviewSceneSettingsTabId, FOnSpawnTab::CreateSP(this, &FStaticMeshEditor::SpawnTab_PreviewSceneSettings))
-		.SetDisplayName(LOCTEXT("PreviewSceneTab", "Preview Scene Settings"))
+	InTabManager->RegisterTabSpawner( PreviewSceneSettingsTabId, FOnSpawnTab::CreateSP(this, &FStaticMeshEditor::SpawnTab_PreviewSceneSettings) )
+		.SetDisplayName( LOCTEXT("PreviewSceneTab", "Preview Scene Settings") )
 		.SetGroup(WorkspaceMenuCategoryRef)
 		.SetIcon(FSlateIcon(FEditorStyle::GetStyleSetName(), "LevelEditor.Tabs.Details"));
+
+	FTabSpawnerEntry& MenuEntry = InTabManager->RegisterTabSpawner( SecondaryToolbarTabId, FOnSpawnTab::CreateSP(this, &FStaticMeshEditor::SpawnTab_SecondaryToolbar) )
+		.SetDisplayName( LOCTEXT("ToolbarTab", "Secondary Toolbar") )
+		.SetGroup(WorkspaceMenuCategoryRef)
+		.SetIcon(FSlateIcon(FEditorStyle::GetStyleSetName(), "Toolbar.Icon"));
+
+	// Hide the menu item by default. It will be enabled only if the secondary toolbar is populated with extensions
+	SecondaryToolbarEntry = &MenuEntry;
+	SecondaryToolbarEntry->SetMenuType( ETabSpawnerMenuType::Hidden );
 }
 
 void FStaticMeshEditor::UnregisterTabSpawners(const TSharedRef<class FTabManager>& InTabManager)
@@ -109,6 +119,7 @@ void FStaticMeshEditor::UnregisterTabSpawners(const TSharedRef<class FTabManager
 	InTabManager->UnregisterTabSpawner( SocketManagerTabId );
 	InTabManager->UnregisterTabSpawner( CollisionTabId );
 	InTabManager->UnregisterTabSpawner( PreviewSceneSettingsTabId );
+	InTabManager->UnregisterTabSpawner( SecondaryToolbarTabId );
 }
 
 
@@ -170,7 +181,7 @@ void FStaticMeshEditor::InitStaticMeshEditor( const EToolkitMode::Type Mode, con
 
 	BuildSubTools();
 
-	const TSharedRef<FTabManager::FLayout> StandaloneDefaultLayout = FTabManager::NewLayout( "Standalone_StaticMeshEditor_Layout_v4" )
+	const TSharedRef<FTabManager::FLayout> StandaloneDefaultLayout = FTabManager::NewLayout( "Standalone_StaticMeshEditor_Layout_v4.1" )
 	->AddArea
 	(
 		FTabManager::NewPrimaryArea() ->SetOrientation(Orient_Vertical)
@@ -180,6 +191,8 @@ void FStaticMeshEditor::InitStaticMeshEditor( const EToolkitMode::Type Mode, con
 			->SetSizeCoefficient(0.1f)
 			->SetHideTabWell( true )
 			->AddTab(GetToolbarTabId(), ETabState::OpenedTab)
+			// Don't want the secondary toolbar tab to be opened if there's nothing in it
+			->AddTab(SecondaryToolbarTabId, ETabState::ClosedTab)
 		)
 		->Split
 		(
@@ -221,6 +234,90 @@ void FStaticMeshEditor::InitStaticMeshEditor( const EToolkitMode::Type Mode, con
 	ExtendMenu();
 	ExtendToolBar();
 	RegenerateMenusAndToolbars();
+	GenerateSecondaryToolbar();
+}
+
+void FStaticMeshEditor::GenerateSecondaryToolbar()
+{
+	// Generate the secondary toolbar only if there are registered extensions
+	TSharedPtr<SDockTab> Tab = TabManager->FindExistingLiveTab(SecondaryToolbarTabId);
+
+	TSharedPtr<FExtender> Extender = FExtender::Combine(SecondaryToolbarExtenders);
+	if (Extender->NumExtensions() == 0)
+	{
+		// If the tab was previously opened, close it since it's now empty
+		if (Tab)
+		{
+			Tab->RemoveTabFromParent();
+		}
+		return;
+	}
+
+	const bool bIsFocusable = true;
+
+	FToolBarBuilder ToolbarBuilder(GetToolkitCommands(), FMultiBoxCustomization::AllowCustomization(GetToolkitFName()), Extender);
+	ToolbarBuilder.SetIsFocusable(bIsFocusable);
+	ToolbarBuilder.BeginSection("Extensions");
+	{
+		// The secondary toolbar itself is empty but will be populated by the extensions when EndSection is called.
+		// The section name helps in the extenders positioning.
+	}
+	ToolbarBuilder.EndSection();
+
+	// Setup the secondary toolbar menu entry
+	SecondaryToolbarEntry->SetMenuType(ETabSpawnerMenuType::Enabled);
+	SecondaryToolbarEntry->SetDisplayName(SecondaryToolbarDisplayName);
+
+	SecondaryToolbar =
+		SNew(SHorizontalBox)
+		+ SHorizontalBox::Slot()
+		.HAlign(HAlign_Left)
+		.AutoWidth()
+		.VAlign(VAlign_Center)
+		[
+			SNew(SVerticalBox)
+			+ SVerticalBox::Slot()
+			.AutoHeight()
+			.VAlign(VAlign_Bottom)
+			[
+				ToolbarBuilder.MakeWidget()
+			]
+		];
+
+	if (SecondaryToolbarWidgetContent.IsValid())
+	{
+		SecondaryToolbarWidgetContent->SetContent(SecondaryToolbar.ToSharedRef());
+	}
+
+	if (!Tab)
+	{
+		// By default, the tab is closed but we want it to be opened when it is populated
+		Tab = TSharedPtr<SDockTab>(TabManager->InvokeTab(SecondaryToolbarTabId));
+	}
+
+	// Override the display name if it was set
+	if (!SecondaryToolbarDisplayName.IsEmpty())
+	{
+		Tab->SetLabel(SecondaryToolbarDisplayName);
+	}
+
+	// But have the focus on the default toolbar
+	TabManager->InvokeTab(GetToolbarTabId());
+}
+
+void FStaticMeshEditor::AddSecondaryToolbarExtender(TSharedPtr<FExtender> Extender)
+{
+	SecondaryToolbarExtenders.AddUnique(Extender);
+}
+
+void FStaticMeshEditor::RemoveSecondaryToolbarExtender(TSharedPtr<FExtender> Extender)
+{
+	SecondaryToolbarExtenders.Remove(Extender);
+}
+
+void FStaticMeshEditor::SetSecondaryToolbarDisplayName(FText DisplayName)
+{
+	SecondaryToolbarDisplayName = DisplayName;
 }
 
 TSharedRef<IDetailCustomization> FStaticMeshEditor::MakeStaticMeshDetails()
@@ -393,16 +490,38 @@ TSharedRef<SDockTab> FStaticMeshEditor::SpawnTab_Collision( const FSpawnTabArgs&
 		];
 }
 
-TSharedRef<SDockTab> FStaticMeshEditor::SpawnTab_PreviewSceneSettings(const FSpawnTabArgs& Args)
+TSharedRef<SDockTab> FStaticMeshEditor::SpawnTab_PreviewSceneSettings( const FSpawnTabArgs& Args )
 {
-	check(Args.GetTabId() == PreviewSceneSettingsTabId);
+	check( Args.GetTabId() == PreviewSceneSettingsTabId );
 	return SNew(SDockTab)
-		.Label(LOCTEXT("StaticMeshPreviewScene_TabTitle", "Preview Scene Settings"))
+		.Label( LOCTEXT("StaticMeshPreviewScene_TabTitle", "Preview Scene Settings") )
 		[
 			AdvancedPreviewSettingsWidget.ToSharedRef()
 		];
 }
 
+TSharedRef<SDockTab> FStaticMeshEditor::SpawnTab_SecondaryToolbar( const FSpawnTabArgs& Args )
+{
+	check( Args.GetTabId() == SecondaryToolbarTabId );
+
+	TSharedRef<SDockTab> SpawnedTab = SNew(SDockTab)
+		.Label( LOCTEXT("SecondaryToolbar_TabTitle", "Secondary Toolbar") )
+		.Icon( FEditorStyle::GetBrush("LevelEditor.Tabs.Toolbar") )
+		.ShouldAutosize( true )
+		[
+			SAssignNew(SecondaryToolbarWidgetContent, SBorder)
+			.Padding(0)
+			.BorderImage(FEditorStyle::GetBrush("NoBorder"))
+		];
+
+	if ( SecondaryToolbar.IsValid() )
+	{
+		SecondaryToolbarWidgetContent->SetContent( SecondaryToolbar.ToSharedRef() );
+	}
+	
+	return SpawnedTab;
+}
+
 void FStaticMeshEditor::BindCommands()
 {
 	const FStaticMeshEditorCommands& Commands = FStaticMeshEditorCommands::Get();
@@ -615,6 +734,7 @@ void FStaticMeshEditor::ExtendToolBar()
 
 	IStaticMeshEditorModule* StaticMeshEditorModule = &FModuleManager::LoadModuleChecked<IStaticMeshEditorModule>( "StaticMeshEditor" );
 	AddToolbarExtender(StaticMeshEditorModule->GetToolBarExtensibilityManager()->GetAllExtenders(GetToolkitCommands(), GetEditingObjects()));
+	AddSecondaryToolbarExtender(StaticMeshEditorModule->GetSecondaryToolBarExtensibilityManager()->GetAllExtenders(GetToolkitCommands(), GetEditingObjects()));
 }
 
 void FStaticMeshEditor::BuildSubTools()
diff --git a/Engine/Source/Editor/StaticMeshEditor/Private/StaticMeshEditor.h b/Engine/Source/Editor/StaticMeshEditor/Private/StaticMeshEditor.h
--- a/Engine/Source/Editor/StaticMeshEditor/Private/StaticMeshEditor.h
+++ b/Engine/Source/Editor/StaticMeshEditor/Private/StaticMeshEditor.h
@@ -31,6 +31,7 @@ class UStaticMesh;
 class UStaticMeshComponent;
 class UStaticMeshSocket;
 struct FPropertyChangedEvent;
+struct FTabSpawnerEntry;
 
 /**
  * StaticMesh Editor class
@@ -43,6 +44,7 @@ public:
 		, MinPrimSize(0.5f)
 		, OverlapNudge(10.0f)
 		, CurrentViewedUVChannel(0)
+		, SecondaryToolbarEntry(nullptr)
 	{}
 
 	~FStaticMeshEditor();
@@ -124,6 +126,8 @@ public:
 	virtual TSet< int32 >& GetSelectedEdges() override;
 
 	virtual FEditorViewportClient& GetViewportClient() override;
+
+	virtual void SetSecondaryToolbarDisplayName(FText DisplayName) override;
 	// End of IStaticMeshEditor
 
 	/** Extends the toolbar menu to include static mesh editor options */
@@ -195,6 +199,7 @@ private:
 	TSharedRef<SDockTab> SpawnTab_SocketManager(const FSpawnTabArgs& Args);
 	TSharedRef<SDockTab> SpawnTab_Collision(const FSpawnTabArgs& Args);
 	TSharedRef<SDockTab> SpawnTab_PreviewSceneSettings(const FSpawnTabArgs& Args);
+	TSharedRef<SDockTab> SpawnTab_SecondaryToolbar(const FSpawnTabArgs& Args);
 
 private:
 	/** Binds commands associated with the Static Mesh Editor. */
@@ -340,6 +345,12 @@ private:
 	bool CanRemoveUVChannel();
 	void RemoveCurrentUVChannel();
 
+	/** Adds or removes extenders to the secondary toolbar */
+	void AddSecondaryToolbarExtender(TSharedPtr<FExtender> Extender);
+	void RemoveSecondaryToolbarExtender(TSharedPtr<FExtender> Extender);
+
+	void GenerateSecondaryToolbar();
+
 private:
 	/** List of open tool panels; used to ensure only one exists at any one time */
 	TMap< FName, TWeakPtr<class SDockableTab> > SpawnedToolPanels;
@@ -402,7 +413,23 @@ private:
 	static const FName SocketManagerTabId;
 	static const FName CollisionTabId;
 	static const FName PreviewSceneSettingsTabId;
+	static const FName SecondaryToolbarTabId;
 
 	/** Allow custom data for this editor */
 	TMap<int32, int32> CustomEditorData;
+
+	/** Static Mesh Editor Secondary Toolbar */
+	TSharedPtr<SWidget> SecondaryToolbar;
+
+	/** The widget that will house the secondary toolbar widget */
+	TSharedPtr<SBorder> SecondaryToolbarWidgetContent;
+
+	/** The extenders to populate the secondary toolbar with */
+	TArray<TSharedPtr<FExtender>> SecondaryToolbarExtenders;
+
+	/** Spawner entry for the secondary toolbar tab */
+	FTabSpawnerEntry* SecondaryToolbarEntry;
+
+	/** The text display name to override the default display name of the secondary toolbar*/
+	FText SecondaryToolbarDisplayName;
 };
diff --git a/Engine/Source/Editor/StaticMeshEditor/Private/StaticMeshEditorModule.cpp b/Engine/Source/Editor/StaticMeshEditor/Private/StaticMeshEditorModule.cpp
--- a/Engine/Source/Editor/StaticMeshEditor/Private/StaticMeshEditorModule.cpp
+++ b/Engine/Source/Editor/StaticMeshEditor/Private/StaticMeshEditorModule.cpp
@@ -28,6 +28,7 @@ public:
 
 		MenuExtensibilityManager = MakeShareable(new FExtensibilityManager);
 		ToolBarExtensibilityManager = MakeShareable(new FExtensibilityManager);
+		SecondaryToolBarExtensibilityManager = MakeShareable(new FExtensibilityManager);
 	}
 
 	/**
@@ -37,6 +38,7 @@ public:
 	{
 		MenuExtensibilityManager.Reset();
 		ToolBarExtensibilityManager.Reset();
+		SecondaryToolBarExtensibilityManager.Reset();
 	}
 
 	/**
@@ -52,10 +54,12 @@ public:
 	/** Gets the extensibility managers for outside entities to extend static mesh editor's menus and toolbars */
 	virtual TSharedPtr<FExtensibilityManager> GetMenuExtensibilityManager() override { return MenuExtensibilityManager; }
 	virtual TSharedPtr<FExtensibilityManager> GetToolBarExtensibilityManager() override { return ToolBarExtensibilityManager; }
+	virtual TSharedPtr<FExtensibilityManager> GetSecondaryToolBarExtensibilityManager() override { return SecondaryToolBarExtensibilityManager; }
 
 private:
 	TSharedPtr<FExtensibilityManager> MenuExtensibilityManager;
 	TSharedPtr<FExtensibilityManager> ToolBarExtensibilityManager;
+	TSharedPtr<FExtensibilityManager> SecondaryToolBarExtensibilityManager;
 };
 
 IMPLEMENT_MODULE( FStaticMeshEditorModule, StaticMeshEditor );
diff --git a/Engine/Source/Editor/StaticMeshEditor/Public/IStaticMeshEditor.h b/Engine/Source/Editor/StaticMeshEditor/Public/IStaticMeshEditor.h
--- a/Engine/Source/Editor/StaticMeshEditor/Public/IStaticMeshEditor.h
+++ b/Engine/Source/Editor/StaticMeshEditor/Public/IStaticMeshEditor.h
@@ -1,4 +1,4 @@
-﻿// Copyright 1998-2018 Epic Games, Inc. All Rights Reserved.
+// Copyright 1998-2018 Epic Games, Inc. All Rights Reserved.
 
 #pragma once
 
@@ -249,6 +249,9 @@ public:
 
 	/** Get the Static Mesh Editor's the viewport client */
 	virtual class FEditorViewportClient& GetViewportClient() = 0;
+
+	/** Set the display name for the secondary toolbar and its associated menu item */
+	virtual void SetSecondaryToolbarDisplayName(FText DisplayName) = 0;
 };
 
 
diff --git a/Engine/Source/Editor/StaticMeshEditor/Public/StaticMeshEditorModule.h b/Engine/Source/Editor/StaticMeshEditor/Public/StaticMeshEditorModule.h
--- a/Engine/Source/Editor/StaticMeshEditor/Public/StaticMeshEditorModule.h
+++ b/Engine/Source/Editor/StaticMeshEditor/Public/StaticMeshEditorModule.h
@@ -22,4 +22,6 @@ public:
 	 * Creates a new static mesh editor.
 	 */
 	virtual TSharedRef<IStaticMeshEditor> CreateStaticMeshEditor( const EToolkitMode::Type Mode, const TSharedPtr< IToolkitHost >& InitToolkitHost, UStaticMesh* StaticMesh ) = 0;
+
+	virtual TSharedPtr<FExtensibilityManager> GetSecondaryToolBarExtensibilityManager() = 0;
 };
diff --git a/Engine/Source/Runtime/Slate/Public/Framework/MultiBox/MultiBoxExtender.h b/Engine/Source/Runtime/Slate/Public/Framework/MultiBox/MultiBoxExtender.h
--- a/Engine/Source/Runtime/Slate/Public/Framework/MultiBox/MultiBoxExtender.h
+++ b/Engine/Source/Runtime/Slate/Public/Framework/MultiBox/MultiBoxExtender.h
@@ -126,6 +126,9 @@ public:
 	 */
 	static SLATE_API TSharedPtr<FExtender> Combine(const TArray< TSharedPtr<FExtender> >& Extenders);
 
+	/** Returns the number of extensions in the extender */
+	SLATE_API int32 NumExtensions() const { return Extensions.Num(); }
+
 private:
 
 	/** List of extensions to our menu.  The order of the list must be maintained. */
