Enables snapping in FStaticMeshEditorViewportClient::TrackingStarted

diff --git a/Engine/Source/Editor/StaticMeshEditor/Private/StaticMeshEditorViewportClient.cpp b/Engine/Source/Editor/StaticMeshEditor/Private/StaticMeshEditorViewportClient.cpp
--- a/Engine/Source/Editor/StaticMeshEditor/Private/StaticMeshEditorViewportClient.cpp
+++ b/Engine/Source/Editor/StaticMeshEditor/Private/StaticMeshEditorViewportClient.cpp
@@ -248,6 +248,7 @@ void FStaticMeshEditorViewportClient::TrackingStarted( const struct FInputEventS
 {
 	if( !bManipulating && bIsDraggingWidget )
 	{
+		Widget->SetSnapEnabled(true);
 		const UStaticMeshSocket* SelectedSocket = StaticMeshEditorPtr.Pin()->GetSelectedSocket();
 		if (SelectedSocket)
 		{
