Static Mesh Editor: Scale preview socket mesh correctly when editor is reopened.

diff --git a/Engine/Source/Editor/StaticMeshEditor/Private/SStaticMeshEditorViewport.cpp b/Engine/Source/Editor/StaticMeshEditor/Private/SStaticMeshEditorViewport.cpp
--- a/Engine/Source/Editor/StaticMeshEditor/Private/SStaticMeshEditorViewport.cpp
+++ b/Engine/Source/Editor/StaticMeshEditor/Private/SStaticMeshEditorViewport.cpp
@@ -177,7 +177,7 @@ void SStaticMeshEditorViewport::UpdatePreviewSocketMeshes()
 					SocketPreviewMeshComponent = NewObject<UStaticMeshComponent>();
 					PreviewScene->AddComponent(SocketPreviewMeshComponent, FTransform::Identity);
 					SocketPreviewMeshComponents.Add(SocketPreviewMeshComponent);
-					SocketPreviewMeshComponent->AttachToComponent(PreviewMeshComponent, FAttachmentTransformRules::SnapToTargetNotIncludingScale, Socket->SocketName);
+					SocketPreviewMeshComponent->AttachToComponent(PreviewMeshComponent, FAttachmentTransformRules::SnapToTargetIncludingScale, Socket->SocketName);
 				}
 				else
 				{
@@ -186,7 +186,7 @@ void SStaticMeshEditorViewport::UpdatePreviewSocketMeshes()
 					// In case of a socket rename, ensure our preview component is still snapping to the proper socket
 					if (!SocketPreviewMeshComponent->GetAttachSocketName().IsEqual(Socket->SocketName))
 					{
-						SocketPreviewMeshComponent->AttachToComponent(PreviewMeshComponent, FAttachmentTransformRules::SnapToTargetNotIncludingScale, Socket->SocketName);
+						SocketPreviewMeshComponent->AttachToComponent(PreviewMeshComponent, FAttachmentTransformRules::SnapToTargetIncludingScale, Socket->SocketName);
 					}
 
 					// Force component to world update to take into account the new socket position.
@@ -249,7 +249,7 @@ void SStaticMeshEditorViewport::UpdatePreviewMesh(UStaticMesh* InStaticMesh, boo
 		{
 			SocketPreviewMeshComponent = NewObject<UStaticMeshComponent>();
 			SocketPreviewMeshComponent->SetStaticMesh(Socket->PreviewStaticMesh);
-			SocketPreviewMeshComponent->AttachToComponent(PreviewMeshComponent, FAttachmentTransformRules::SnapToTargetNotIncludingScale, Socket->SocketName);
+			SocketPreviewMeshComponent->AttachToComponent(PreviewMeshComponent, FAttachmentTransformRules::SnapToTargetIncludingScale, Socket->SocketName);
 			SocketPreviewMeshComponents.Add(SocketPreviewMeshComponent);
 			PreviewScene->AddComponent(SocketPreviewMeshComponent, FTransform::Identity);
 		}
