Deprecate FRawMesh - Miscellaneous stuff

diff --git a/Engine/Source/Editor/StaticMeshEditor/Private/StaticMeshEditorViewportClient.cpp b/Engine/Source/Editor/StaticMeshEditor/Private/StaticMeshEditorViewportClient.cpp
--- a/Engine/Source/Editor/StaticMeshEditor/Private/StaticMeshEditorViewportClient.cpp
+++ b/Engine/Source/Editor/StaticMeshEditor/Private/StaticMeshEditorViewportClient.cpp
@@ -16,7 +16,6 @@
 #include "UnrealEngine.h"
 
 #include "StaticMeshResources.h"
-#include "RawMesh.h"
 #include "DistanceFieldAtlas.h"
 #include "SEditorViewport.h"
 #include "AdvancedPreviewScene.h"
@@ -1012,172 +1011,191 @@ void FStaticMeshEditorViewportClient::ProcessClick(class FSceneView& InView, cla
 				FVector ClosestEdgeVertices[ 2 ];
 
 				const uint32 LODLevel = FMath::Clamp( StaticMeshComponent->ForcedLodModel - 1, 0, StaticMeshComponent->GetStaticMesh()->GetNumLODs() - 1 );
-				FRawMesh RawMesh;
-				StaticMeshComponent->GetStaticMesh()->SourceModels[LODLevel].LoadRawMesh(RawMesh);
-
-				const int32 RawEdgeCount = RawMesh.WedgeIndices.Num() - 1; 
-				const int32 NumFaces = RawMesh.WedgeIndices.Num() / 3;
-				int32 NumBackFacingTriangles = 0;
-				for(int32 FaceIndex = 0; FaceIndex < NumFaces; ++FaceIndex)
+				if (StaticMeshComponent->GetStaticMesh()->HasValidRenderData(true, LODLevel))
 				{
-					// We disable edge selection where all adjoining triangles are back face culled and the 
-					// material is not two-sided. This prevents edges that are back-face culled from being selected.
-					bool bIsBackFacing = false;
-					bool bIsTwoSided = false;
-					UMaterialInterface* Material = StaticMeshComponent->GetMaterial(RawMesh.FaceMaterialIndices[FaceIndex]);
-					if (Material && Material->GetMaterial())
-					{
-						bIsTwoSided = Material->IsTwoSided();
-					}
-					if(!bIsTwoSided)
-					{
-						// Check whether triangle if back facing 
-						const FVector A = RawMesh.GetWedgePosition( FaceIndex * 3);
-						const FVector B = RawMesh.GetWedgePosition( FaceIndex * 3 + 1);
-						const FVector C = RawMesh.GetWedgePosition( FaceIndex * 3 + 2);
-								
-						// Compute the per-triangle normal
-						const FVector BA = A - B;
-						const FVector CA = A - C;
-						const FVector TriangleNormal = (CA ^ BA).GetSafeNormal();
-
-						// Transform the view position from world to component space
-						const FVector ComponentSpaceViewOrigin = StaticMeshComponent->GetComponentTransform().InverseTransformPosition( View->ViewMatrices.GetViewOrigin());
-								
-						// Determine which side of the triangle's plane that the view position lies on.
-						bIsBackFacing = (FVector::PointPlaneDist( ComponentSpaceViewOrigin,  A, TriangleNormal)  < 0.0f);
-					}
-						
-					for( int32 VertIndex = 0; VertIndex < 3; ++VertIndex )
+					FStaticMeshLODResources& RenderData = StaticMeshComponent->GetStaticMesh()->RenderData->LODResources[LODLevel];
+
+					int32 NumBackFacingTriangles = 0;
+					uint32 IndexBufferIndex = 0;
+					for (int32 SectionIndex = 0; SectionIndex < RenderData.Sections.Num(); ++SectionIndex)
 					{
-						const int32 EdgeIndex = FaceIndex * 3 + VertIndex;
-						const int32 EdgeIndex2 = FaceIndex * 3 + ((VertIndex + 1) % 3);
-
-						FVector EdgeVertices[ 2 ];
-						EdgeVertices[0]	= RawMesh.GetWedgePosition(EdgeIndex);
-						EdgeVertices[1] = RawMesh.GetWedgePosition(EdgeIndex2);
-
-						// First check to see if this edge is already in our "closest to click" list.
-						// Most edges are shared by two faces in our raw triangle data set, so we want
-						// to select (or deselect) both of these edges that the user clicks on (what
-						// appears to be) a single edge
-						if( ClosestEdgeIndices.Num() > 0 &&
-							( ( EdgeVertices[ 0 ].Equals( ClosestEdgeVertices[ 0 ] ) && EdgeVertices[ 1 ].Equals( ClosestEdgeVertices[ 1 ] ) ) ||
-							( EdgeVertices[ 0 ].Equals( ClosestEdgeVertices[ 1 ] ) && EdgeVertices[ 1 ].Equals( ClosestEdgeVertices[ 0 ] ) ) ) )
+						const FStaticMeshSection& Section = RenderData.Sections[SectionIndex];
+						const int32 FaceMaterialIndex = Section.MaterialIndex;
+						const int32 NumFaces = Section.NumTriangles;
+						for (int32 FaceIndex = 0; FaceIndex < NumFaces; ++FaceIndex)
 						{
-							// Edge overlaps the closest edge we have so far, so just add it to the list
-							ClosestEdgeIndices.Add( EdgeIndex );
-							// Increment the number of back facing triangles if the adjoining triangle 
-							// is back facing and isn't two-sided
-							if(bIsBackFacing && !bIsTwoSided)
+							FVector VertexPosition[3];
+							uint32 VertexIndex[3];
+							uint32 WedgeIndex[3];
+							for (int32 Corner = 0; Corner < 3; ++Corner)
 							{
-								++NumBackFacingTriangles;
+								WedgeIndex[Corner] = IndexBufferIndex;
+								VertexIndex[Corner] = RenderData.IndexBuffer.GetIndex(IndexBufferIndex);
+								VertexPosition[Corner] = RenderData.VertexBuffers.PositionVertexBuffer.VertexPosition(VertexIndex[Corner]);
+								IndexBufferIndex++;
 							}
-						}
-						else
-						{
-							FVector WorldSpaceEdgeStart( StaticMeshComponent->GetComponentTransform().TransformPosition( EdgeVertices[ 0 ] ) );
-							FVector WorldSpaceEdgeEnd( StaticMeshComponent->GetComponentTransform().TransformPosition( EdgeVertices[ 1 ] ) );
-
-							// Determine the mesh edge that's closest to the ray cast through the eye towards the click location
-							FVector ClosestPointToEdgeOnClickLine;
-							FVector ClosestPointToClickLineOnEdge;
-							FMath::SegmentDistToSegment(
-								ClickLineStart,
-								ClickLineEnd,
-								WorldSpaceEdgeStart,
-								WorldSpaceEdgeEnd,
-								ClosestPointToEdgeOnClickLine,
-								ClosestPointToClickLineOnEdge );
-
-							// Compute the minimum distance (squared)
-							const float MinDistanceToEdgeSquared = ( ClosestPointToClickLineOnEdge - ClosestPointToEdgeOnClickLine ).SizeSquared();
-
-							if( MinDistanceToEdgeSquared <= WorldSpaceMinClickDistance )
+							// We disable edge selection where all adjoining triangles are back face culled and the 
+							// material is not two-sided. This prevents edges that are back-face culled from being selected.
+							bool bIsBackFacing = false;
+							bool bIsTwoSided = false;
+							UMaterialInterface* Material = StaticMeshComponent->GetMaterial(FaceMaterialIndex);
+							if (Material && Material->GetMaterial())
 							{
-								if( MinDistanceToEdgeSquared <= ClosestEdgeDistance )
-								{
-									// This is the closest edge to the click line that we've found so far!
-									ClosestEdgeDistance = MinDistanceToEdgeSquared;
-									ClosestEdgeVertices[ 0 ] = EdgeVertices[ 0 ];
-									ClosestEdgeVertices[ 1 ] = EdgeVertices[ 1 ];
+								bIsTwoSided = Material->IsTwoSided();
+							}
+							if (!bIsTwoSided)
+							{
+								// Check whether triangle if back facing 
+								const FVector A = VertexPosition[0];
+								const FVector B = VertexPosition[1];
+								const FVector C = VertexPosition[2];
+
+								// Compute the per-triangle normal
+								const FVector BA = A - B;
+								const FVector CA = A - C;
+								const FVector TriangleNormal = (CA ^ BA).GetSafeNormal();
 
-									ClosestEdgeIndices.Reset();
-									ClosestEdgeIndices.Add( EdgeIndex );
+								// Transform the view position from world to component space
+								const FVector ComponentSpaceViewOrigin = StaticMeshComponent->GetComponentTransform().InverseTransformPosition(View->ViewMatrices.GetViewOrigin());
 
-									// Reset the number of back facing triangles.
-									NumBackFacingTriangles = (bIsBackFacing && !bIsTwoSided) ? 1 : 0;
+								// Determine which side of the triangle's plane that the view position lies on.
+								bIsBackFacing = (FVector::PointPlaneDist(ComponentSpaceViewOrigin, A, TriangleNormal) < 0.0f);
+							}
+
+							for (int32 Corner = 0; Corner < 3; ++Corner)
+							{
+								const int32 Corner2 = (Corner + 1) % 3;
+
+								FVector EdgeVertices[2];
+								EdgeVertices[0] = VertexPosition[Corner];
+								EdgeVertices[1] = VertexPosition[Corner2];
+
+								// First check to see if this edge is already in our "closest to click" list.
+								// Most edges are shared by two faces in our raw triangle data set, so we want
+								// to select (or deselect) both of these edges that the user clicks on (what
+								// appears to be) a single edge
+								if (ClosestEdgeIndices.Num() > 0 &&
+									((EdgeVertices[0].Equals(ClosestEdgeVertices[0]) && EdgeVertices[1].Equals(ClosestEdgeVertices[1])) ||
+									(EdgeVertices[0].Equals(ClosestEdgeVertices[1]) && EdgeVertices[1].Equals(ClosestEdgeVertices[0]))))
+								{
+									// Edge overlaps the closest edge we have so far, so just add it to the list
+									ClosestEdgeIndices.Add(WedgeIndex[Corner]);
+									// Increment the number of back facing triangles if the adjoining triangle 
+									// is back facing and isn't two-sided
+									if (bIsBackFacing && !bIsTwoSided)
+									{
+										++NumBackFacingTriangles;
+									}
+								}
+								else
+								{
+									FVector WorldSpaceEdgeStart(StaticMeshComponent->GetComponentTransform().TransformPosition(EdgeVertices[0]));
+									FVector WorldSpaceEdgeEnd(StaticMeshComponent->GetComponentTransform().TransformPosition(EdgeVertices[1]));
+
+									// Determine the mesh edge that's closest to the ray cast through the eye towards the click location
+									FVector ClosestPointToEdgeOnClickLine;
+									FVector ClosestPointToClickLineOnEdge;
+									FMath::SegmentDistToSegment(
+										ClickLineStart,
+										ClickLineEnd,
+										WorldSpaceEdgeStart,
+										WorldSpaceEdgeEnd,
+										ClosestPointToEdgeOnClickLine,
+										ClosestPointToClickLineOnEdge);
+
+									// Compute the minimum distance (squared)
+									const float MinDistanceToEdgeSquared = (ClosestPointToClickLineOnEdge - ClosestPointToEdgeOnClickLine).SizeSquared();
+
+									if (MinDistanceToEdgeSquared <= WorldSpaceMinClickDistance)
+									{
+										if (MinDistanceToEdgeSquared <= ClosestEdgeDistance)
+										{
+											// This is the closest edge to the click line that we've found so far!
+											ClosestEdgeDistance = MinDistanceToEdgeSquared;
+											ClosestEdgeVertices[0] = EdgeVertices[0];
+											ClosestEdgeVertices[1] = EdgeVertices[1];
+
+											ClosestEdgeIndices.Reset();
+											ClosestEdgeIndices.Add(WedgeIndex[Corner]);
+
+											// Reset the number of back facing triangles.
+											NumBackFacingTriangles = (bIsBackFacing && !bIsTwoSided) ? 1 : 0;
+										}
+									}
 								}
 							}
 						}
 					}
-				}
 
-				// Did the user click on an edge? Edges must also have at least one adjoining triangle 
-				// which isn't back face culled (for one-sided materials)
-				if( ClosestEdgeIndices.Num() > 0 && ClosestEdgeIndices.Num() > NumBackFacingTriangles)
-				{
-					for( int32 CurIndex = 0; CurIndex < ClosestEdgeIndices.Num(); ++CurIndex )
-					{
-						const int32 CurEdgeIndex = ClosestEdgeIndices[ CurIndex ];
 
-						if( bCtrlDown )
+					// Did the user click on an edge? Edges must also have at least one adjoining triangle 
+					// which isn't back face culled (for one-sided materials)
+					if (ClosestEdgeIndices.Num() > 0 && ClosestEdgeIndices.Num() > NumBackFacingTriangles)
+					{
+						for (int32 CurIndex = 0; CurIndex < ClosestEdgeIndices.Num(); ++CurIndex)
 						{
-							// Toggle selection
-							if( SelectedEdgeIndices.Contains( CurEdgeIndex ) )
+							const int32 CurEdgeIndex = ClosestEdgeIndices[CurIndex];
+
+							if (bCtrlDown)
 							{
-								SelectedEdgeIndices.Remove( CurEdgeIndex );
+								// Toggle selection
+								if (SelectedEdgeIndices.Contains(CurEdgeIndex))
+								{
+									SelectedEdgeIndices.Remove(CurEdgeIndex);
+								}
+								else
+								{
+									SelectedEdgeIndices.Add(CurEdgeIndex);
+								}
 							}
 							else
 							{
-								SelectedEdgeIndices.Add( CurEdgeIndex );
+								// Append to selection
+								SelectedEdgeIndices.Add(CurEdgeIndex);
 							}
 						}
-						else
+
+						// Reset cached vertices and uv coordinates.
+						SelectedEdgeVertices.Reset();
+						for (int32 TexCoordIndex = 0; TexCoordIndex < MAX_STATIC_TEXCOORDS; ++TexCoordIndex)
 						{
-							// Append to selection
-							SelectedEdgeIndices.Add( CurEdgeIndex );
+							SelectedEdgeTexCoords[TexCoordIndex].Reset();
 						}
-					}
 
-					// Reset cached vertices and uv coordinates.
-					SelectedEdgeVertices.Reset();
-					for(int32 TexCoordIndex = 0; TexCoordIndex < MAX_STATIC_TEXCOORDS; ++TexCoordIndex)
-					{
-						SelectedEdgeTexCoords[TexCoordIndex].Reset();
-					}
-
-					for(FSelectedEdgeSet::TIterator SelectionIt( SelectedEdgeIndices ); SelectionIt; ++SelectionIt)
-					{
-						const uint32 EdgeIndex = *SelectionIt;
-						const uint32 FaceIndex = EdgeIndex / 3;
+						for (FSelectedEdgeSet::TIterator SelectionIt(SelectedEdgeIndices); SelectionIt; ++SelectionIt)
+						{
+							const uint32 EdgeIndex = *SelectionIt;
+							const uint32 FaceIndex = EdgeIndex / 3;
 
-						const uint32 WedgeIndex = FaceIndex * 3 + (EdgeIndex % 3);
-						const uint32 WedgeIndex2 = FaceIndex * 3 + ((EdgeIndex + 1) % 3);
+							const uint32 WedgeIndex = FaceIndex * 3 + (EdgeIndex % 3);
+							const uint32 WedgeIndex2 = FaceIndex * 3 + ((EdgeIndex + 1) % 3);
 
-						// Cache edge vertices in local space.
-						FVector EdgeVertices[ 2 ];
-						EdgeVertices[ 0 ] = RawMesh.GetWedgePosition(WedgeIndex);
-						EdgeVertices[ 1 ] = RawMesh.GetWedgePosition(WedgeIndex2);
+							const uint32 VertexIndex = RenderData.IndexBuffer.GetIndex(WedgeIndex);
+							const uint32 VertexIndex2 = RenderData.IndexBuffer.GetIndex(WedgeIndex2);
+							// Cache edge vertices in local space.
+							FVector EdgeVertices[2];
+							EdgeVertices[0] = RenderData.VertexBuffers.PositionVertexBuffer.VertexPosition(VertexIndex);
+							EdgeVertices[1] = RenderData.VertexBuffers.PositionVertexBuffer.VertexPosition(VertexIndex2);
 
-						SelectedEdgeVertices.Add(EdgeVertices[0]);
-						SelectedEdgeVertices.Add(EdgeVertices[1]);
+							SelectedEdgeVertices.Add(EdgeVertices[0]);
+							SelectedEdgeVertices.Add(EdgeVertices[1]);
 
-						// Cache UV
-						for(int32 TexCoordIndex = 0; TexCoordIndex < MAX_STATIC_TEXCOORDS; ++TexCoordIndex)
-						{
-							if( RawMesh.WedgeTexCoords[TexCoordIndex].Num() > 0)
+							// Cache UV
+							for (uint32 TexCoordIndex = 0; TexCoordIndex < MAX_STATIC_TEXCOORDS; ++TexCoordIndex)
 							{
-								FVector2D UVIndex1, UVIndex2;
-								UVIndex1 = RawMesh.WedgeTexCoords[TexCoordIndex][WedgeIndex];
-								UVIndex2 = RawMesh.WedgeTexCoords[TexCoordIndex][WedgeIndex2];
-								SelectedEdgeTexCoords[TexCoordIndex].Add(UVIndex1);
-								SelectedEdgeTexCoords[TexCoordIndex].Add(UVIndex2);
+								if (RenderData.VertexBuffers.StaticMeshVertexBuffer.GetNumTexCoords() > TexCoordIndex)
+								{
+									FVector2D UVIndex1, UVIndex2;
+									UVIndex1 = RenderData.VertexBuffers.StaticMeshVertexBuffer.GetVertexUV(VertexIndex, TexCoordIndex);
+									UVIndex2 = RenderData.VertexBuffers.StaticMeshVertexBuffer.GetVertexUV(VertexIndex2, TexCoordIndex);
+									SelectedEdgeTexCoords[TexCoordIndex].Add(UVIndex1);
+									SelectedEdgeTexCoords[TexCoordIndex].Add(UVIndex2);
+								}
 							}
 						}
-					}
 
-					ClearSelectedEdges = false;
+						ClearSelectedEdges = false;
+					}
 				}
 			}
 		}
diff --git a/Engine/Source/Editor/UnrealEd/Private/Fbx/FbxMainExport.cpp b/Engine/Source/Editor/UnrealEd/Private/Fbx/FbxMainExport.cpp
--- a/Engine/Source/Editor/UnrealEd/Private/Fbx/FbxMainExport.cpp
+++ b/Engine/Source/Editor/UnrealEd/Private/Fbx/FbxMainExport.cpp
@@ -66,7 +66,12 @@
 #include "Matinee/InterpGroupInst.h"
 #include "Matinee/MatineeActor.h"
 #include "FbxExporter.h"
-#include "RawMesh.h"
+
+#include "MeshDescription.h"
+#include "MeshAttributes.h"
+#include "MeshAttributeArray.h"
+#include "MeshDescriptionOperations.h"
+
 #include "Components/BrushComponent.h"
 #include "CineCameraComponent.h"
 #include "Math/UnitConversion.h"
@@ -661,11 +666,12 @@ void FFbxExporter::ExportBrush(ABrush* Actor, UModel* InModel, bool bConvertToSt
 	}
 	else
 	{
-		FRawMesh Mesh;
+		FMeshDescription Mesh;
+		UStaticMesh::RegisterMeshAttributes(Mesh);
 		TArray<FStaticMaterial>	Materials;
 		GetBrushMesh(Actor,Actor->Brush,Mesh,Materials);
 
-		if( Mesh.VertexPositions.Num() )
+		if( Mesh.Vertices().Num() )
 		{
 			UStaticMesh* StaticMesh = CreateStaticMesh(Mesh,Materials,GetTransientPackage(),Actor->GetFName());
 			ExportStaticMesh( StaticMesh, &Materials );
@@ -907,8 +913,9 @@ void FFbxExporter::ExportStaticMesh(AActor* Actor, UStaticMeshComponent* StaticM
 
 struct FBSPExportData
 {
-	FRawMesh Mesh;
+	FMeshDescription Mesh;
 	TArray<FStaticMaterial> Materials;
+	TArray<uint32> SmoothGroups;
 	uint32 NumVerts;
 	uint32 NumFaces;
 	uint32 CurrentVertAddIndex;
@@ -944,8 +951,15 @@ void FFbxExporter::ExportBSP( UModel* Model, bool bSelectedOnly )
 			{
 				FBSPExportData& Data = BrushToMeshMap.FindOrAdd( BrushActor );
 
+				FPoly Poly;
+				GEditor->polyFindMaster(Model, Node.iSurf, Poly);
+
 				Data.NumVerts += Node.NumVertices;
 				Data.NumFaces += Node.NumVertices-2;
+				UMaterialInterface*	Material = Poly.Material;
+				FName MaterialName = Material != nullptr ? Material->GetFName() : NAME_None;
+				Data.Materials.AddUnique(FStaticMaterial(Material, MaterialName, MaterialName));
+				AllMaterials.AddUnique(FStaticMaterial(Material, MaterialName, MaterialName));
 			}
 		}
 	}
@@ -970,36 +984,53 @@ void FFbxExporter::ExportBSP( UModel* Model, bool bSelectedOnly )
 			}
 
 			TArray<FStaticMaterial>& Materials = ExportData->Materials;
-			FRawMesh& Mesh = ExportData->Mesh;
+			FMeshDescription& Mesh = ExportData->Mesh;
 
 			//Pre-allocate space for this mesh.
 			if( !ExportData->bInitialised )
 			{
-				ExportData->bInitialised = true;
-				Mesh.VertexPositions.Empty();
-				Mesh.VertexPositions.AddUninitialized(ExportData->NumVerts);
+				uint32 NumWedges = ExportData->NumFaces * 3;
 
-				Mesh.FaceMaterialIndices.Empty();
-				Mesh.FaceMaterialIndices.AddUninitialized(ExportData->NumFaces);
-				Mesh.FaceSmoothingMasks.Empty();
-				Mesh.FaceSmoothingMasks.AddUninitialized(ExportData->NumFaces);
-				
-				uint32 NumWedges = ExportData->NumFaces*3;
-				Mesh.WedgeIndices.Empty();
-				Mesh.WedgeIndices.AddUninitialized(NumWedges);
-				Mesh.WedgeTexCoords[0].Empty();
-				Mesh.WedgeTexCoords[0].AddUninitialized(NumWedges);
-				Mesh.WedgeColors.Empty();
-				Mesh.WedgeColors.AddUninitialized(NumWedges);
-				Mesh.WedgeTangentZ.Empty();
-				Mesh.WedgeTangentZ.AddUninitialized(NumWedges);
+				UStaticMesh::RegisterMeshAttributes(Mesh);
+				ExportData->bInitialised = true;
+				Mesh.Empty();
+				Mesh.ReserveNewVertices(ExportData->NumVerts);
+				Mesh.ReserveNewVertexInstances(NumWedges);
+				Mesh.ReserveNewPolygons(ExportData->NumFaces);
+				Mesh.ReserveNewEdges(ExportData->NumFaces * 2);
+				Mesh.ReserveNewPolygonGroups(Materials.Num());
+				//We need to get the PolygonGroupImportedMaterial
+				FStaticMeshDescriptionAttributeGetter InitMeshAttributes(&Mesh);
+				TPolygonGroupAttributesRef<FName> InitPolygonGroupNames = InitMeshAttributes.GetPolygonGroupImportedMaterialSlotNames();
+				for (const FStaticMaterial& StaticMaterial : Materials)
+				{
+					const FPolygonGroupID PolygonGroupID = Mesh.CreatePolygonGroup();
+					InitPolygonGroupNames[PolygonGroupID] = StaticMaterial.ImportedMaterialSlotName;
+				}
+				ExportData->SmoothGroups.Empty(ExportData->NumFaces);
 			}
 			
-			UMaterialInterface*	Material = Poly.Material;
+			//Get the Attributes
+			FStaticMeshDescriptionAttributeGetter MeshAttributes(&Mesh);
+			TVertexAttributesRef<FVector> VertexPositions = MeshAttributes.GetPositions();
+			TVertexInstanceAttributesRef<FVector2D> UVs = MeshAttributes.GetUVs();
+			TVertexInstanceAttributesRef<FVector4> Colors = MeshAttributes.GetColors();
+			TVertexInstanceAttributesRef<FVector> Normals = MeshAttributes.GetNormals();
+			TEdgeAttributesRef<bool> EdgeHardnesses = MeshAttributes.GetEdgeHardnesses();
+			TEdgeAttributesRef<float> EdgeCreaseSharpnesses = MeshAttributes.GetEdgeCreaseSharpnesses();
 
-			AllMaterials.AddUnique(FStaticMaterial(Material));
-
-			int32 MaterialIndex = ExportData->Materials.AddUnique(Material);
+			
+			UMaterialInterface*	Material = Poly.Material;
+			FName MaterialName = Material != nullptr ? Material->GetFName() : NAME_None;
+			
+			int32 MaterialIndex = INDEX_NONE;
+			if (!ExportData->Materials.Find(FStaticMaterial(Material, MaterialName, MaterialName), MaterialIndex))
+			{
+				MaterialIndex = 0;
+			}
+			const FPolygonGroupID CurrentPolygonGroupID(MaterialIndex);
+			//The material ID should follow the unique materials
+			check(Mesh.IsPolygonGroupValid(CurrentPolygonGroupID));
 
 			const FVector& TextureBase = Model->Points[Surf.pBase];
 			const FVector& TextureX = Model->Vectors[Surf.vTextureU];
@@ -1012,12 +1043,15 @@ void FFbxExporter::ExportBSP( UModel* Model, bool bSelectedOnly )
 			{
 				const FVert& Vert = Model->Verts[Node.iVertPool + VertexIndex];
 				const FVector& Vertex = Model->Points[Vert.pVertex];
-				Mesh.VertexPositions[ExportData->CurrentVertAddIndex+VertexIndex] = Vertex;
+				FVertexID VertexID = Mesh.CreateVertex();
+				VertexPositions[VertexID] = Vertex;
 			}
 			ExportData->CurrentVertAddIndex += Node.NumVertices;
 
 			for (int32 StartVertexIndex = 1; StartVertexIndex < Node.NumVertices - 1; ++StartVertexIndex)
 			{
+				FVertexInstanceID VertexInstanceIDs[3];
+				FVertexID VertexIDs[3];
 				// These map the node's vertices to the 3 triangle indices to triangulate the convex polygon.
 				int32 TriVertIndices[3] = {
 					Node.iVertPool + StartVertexIndex + 1,
@@ -1029,26 +1063,54 @@ void FFbxExporter::ExportBSP( UModel* Model, bool bSelectedOnly )
 					StartIndex + StartVertexIndex,
 					StartIndex };
 
-				Mesh.FaceMaterialIndices[ExportData->CurrentFaceAddIndex] = MaterialIndex;
-				Mesh.FaceSmoothingMasks[ExportData->CurrentFaceAddIndex] =  ( 1 << ( Node.iSurf % 32 ) );
-
-				for (uint32 WedgeIndex = 0; WedgeIndex < 3; ++WedgeIndex)
+				for (uint32 Corner = 0; Corner < 3; ++Corner)
 				{
-					const FVert& Vert = Model->Verts[TriVertIndices[WedgeIndex]];
+					VertexIDs[Corner] = FVertexID(WedgeIndices[Corner]);
+					VertexInstanceIDs[Corner] = Mesh.CreateVertexInstance(VertexIDs[Corner]);
+					const FVert& Vert = Model->Verts[TriVertIndices[Corner]];
 					const FVector& Vertex = Model->Points[Vert.pVertex];
 
 					float U = ((Vertex - TextureBase) | TextureX) / UModel::GetGlobalBSPTexelScale();
 					float V = ((Vertex - TextureBase) | TextureY) / UModel::GetGlobalBSPTexelScale();
-
-					uint32 RealWedgeIndex = ( ExportData->CurrentFaceAddIndex * 3 ) + WedgeIndex;
-
-					Mesh.WedgeIndices[RealWedgeIndex] = WedgeIndices[WedgeIndex];
-					Mesh.WedgeTexCoords[0][RealWedgeIndex] = FVector2D(U,V);
+					UVs.Set(VertexInstanceIDs[Corner], 0, FVector2D(U, V));
 					//This is not exported when exporting the whole level via ExportModel so leaving out here for now. 
-					//Mesh.WedgeTexCoords[1][RealWedgeIndex] = Vert.ShadowTexCoord;
-					Mesh.WedgeColors[RealWedgeIndex] = FColor(255,255,255,255);
-					Mesh.WedgeTangentZ[RealWedgeIndex] = Normal;
+					//UVs.Set(VertexInstanceIDs[Corner], 1, Vert.ShadowTexCoord);
+					Colors[VertexInstanceIDs[Corner]] = FLinearColor::White;
+					Normals[VertexInstanceIDs[Corner]] = Normal;
 				}
+				//Add triangles
+				TArray<FMeshDescription::FContourPoint> Contours;
+				for (int32 Corner = 0; Corner < 3; ++Corner)
+				{
+					int32 ContourPointIndex = Contours.AddDefaulted();
+					FMeshDescription::FContourPoint& ContourPoint = Contours[ContourPointIndex];
+					//Find the matching edge ID
+					uint32 CornerIndices[2];
+					CornerIndices[0] = (Corner + 0) % 3;
+					CornerIndices[1] = (Corner + 1) % 3;
+
+					FVertexID EdgeVertexIDs[2];
+					EdgeVertexIDs[0] = VertexIDs[CornerIndices[0]];
+					EdgeVertexIDs[1] = VertexIDs[CornerIndices[1]];
+
+					FEdgeID MatchEdgeId = Mesh.GetVertexPairEdge(EdgeVertexIDs[0], EdgeVertexIDs[1]);
+					if (MatchEdgeId == FEdgeID::Invalid)
+					{
+						MatchEdgeId = Mesh.CreateEdge(EdgeVertexIDs[0], EdgeVertexIDs[1]);
+						EdgeHardnesses[MatchEdgeId] = false;
+						EdgeCreaseSharpnesses[MatchEdgeId] = 0.0f;
+					}
+					ContourPoint.EdgeID = MatchEdgeId;
+					ContourPoint.VertexInstanceID = VertexInstanceIDs[CornerIndices[0]];
+				}
+				// Insert a polygon into the mesh
+				const FPolygonID NewPolygonID = Mesh.CreatePolygon(CurrentPolygonGroupID, Contours);
+				//Triangulate the polygon
+				FMeshPolygon& Polygon = Mesh.GetPolygon(NewPolygonID);
+				Mesh.ComputePolygonTriangulation(NewPolygonID, Polygon.Triangles);
+
+				//Add to the smoothGroup array so we can compute hard edge later
+				ExportData->SmoothGroups.Add((1 << (Node.iSurf % 32)));
 
 				++ExportData->CurrentFaceAddIndex;
 			}
@@ -1057,11 +1119,13 @@ void FFbxExporter::ExportBSP( UModel* Model, bool bSelectedOnly )
 
 	for( TMap< ABrush*, FBSPExportData >::TIterator It(BrushToMeshMap); It; ++It )
 	{
-		if( It.Value().Mesh.VertexPositions.Num() )
+		if( It.Value().Mesh.Vertices().Num() )
 		{
+			FMeshDescriptionOperations::ConvertSmoothGroupToHardEdges(It.Value().SmoothGroups, It.Value().Mesh);
+
 			UStaticMesh* NewMesh = CreateStaticMesh( It.Value().Mesh, It.Value().Materials, GetTransientPackage(), It.Key()->GetFName() );
 
-			ExportStaticMesh( NewMesh, &AllMaterials );
+			ExportStaticMesh( NewMesh, &It.Value().Materials);
 		}
 	}
 }
diff --git a/Engine/Source/Editor/UnrealEd/Private/Fbx/FbxStaticMeshImport.cpp b/Engine/Source/Editor/UnrealEd/Private/Fbx/FbxStaticMeshImport.cpp
--- a/Engine/Source/Editor/UnrealEd/Private/Fbx/FbxStaticMeshImport.cpp
+++ b/Engine/Source/Editor/UnrealEd/Private/Fbx/FbxStaticMeshImport.cpp
@@ -21,7 +21,6 @@
 #include "Engine/StaticMeshSocket.h"
 #include "Editor.h"
 #include "Modules/ModuleManager.h"
-#include "RawMesh.h"
 
 #include "StaticMeshResources.h"
 #include "ObjectTools.h"
@@ -204,7 +203,7 @@ struct FFBXUVs
 				}
 			}
 		}
-		UniqueUVCount = FMath::Min<int32>(UniqueUVCount, MAX_MESH_TEXTURE_COORDS);
+		UniqueUVCount = FMath::Min<int32>(UniqueUVCount, MAX_MESH_TEXTURE_COORDS_MD);
 	}
 
 	int32 FindLightUVIndex() const
@@ -261,7 +260,7 @@ struct FFBXUVs
 };
 
 
-bool UnFbx::FFbxImporter::BuildStaticMeshFromGeometry(FbxNode* Node, UStaticMesh* StaticMesh, TArray<FFbxMaterial>& MeshMaterials, int32 LODIndex,FRawMesh& RawMesh,
+bool UnFbx::FFbxImporter::BuildStaticMeshFromGeometry(FbxNode* Node, UStaticMesh* StaticMesh, TArray<FFbxMaterial>& MeshMaterials, int32 LODIndex,
 	EVertexColorImportOption::Type VertexColorImportOption, const TMap<FVector, FColor>& ExistingVertexColorData, const FColor& VertexOverrideColor)
 {
 	check(StaticMesh->SourceModels.IsValidIndex(LODIndex));
@@ -548,7 +547,7 @@ bool UnFbx::FFbxImporter::BuildStaticMeshFromGeometry(FbxNode* Node, UStaticMesh
 	int32 ExistingUVCount = VertexInstanceUVs.GetNumIndices();
 
 	int32 NumUVs = FMath::Max(FBXUVs.UniqueUVCount, ExistingUVCount);
-	NumUVs = FMath::Min<int32>(MAX_MESH_TEXTURE_COORDS, NumUVs);
+	NumUVs = FMath::Min<int32>(MAX_MESH_TEXTURE_COORDS_MD, NumUVs);
 	// At least one UV set must exist.  
 	NumUVs = FMath::Max(1, NumUVs);
 
@@ -1593,12 +1592,6 @@ UStaticMesh* UnFbx::FFbxImporter::ImportStaticMeshAsSingle(UObject* InParent, TA
 	}
 
 	FStaticMeshSourceModel& SrcModel = StaticMesh->SourceModels[LODIndex];
-	if( InStaticMesh != NULL && LODIndex > 0 && !SrcModel.RawMeshBulkData->IsEmpty() )
-	{
-		// clear out the old mesh data
-		FRawMesh EmptyRawMesh;
-		SrcModel.SaveRawMesh(EmptyRawMesh);
-	}
 	
 	// make sure it has a new lighting guid
 	StaticMesh->LightingGuid = FGuid::NewGuid();
@@ -1608,9 +1601,6 @@ UStaticMesh* UnFbx::FFbxImporter::ImportStaticMeshAsSingle(UObject* InParent, TA
 	StaticMesh->LightMapCoordinateIndex = 1;
 
 
-	FRawMesh NewRawMesh;
-	SrcModel.RawMeshBulkData->LoadRawMesh(NewRawMesh);
-
 	TArray<FFbxMaterial> MeshMaterials;
 	for (MeshIndex = 0; MeshIndex < MeshNodeArray.Num(); MeshIndex++ )
 	{
@@ -1618,7 +1608,7 @@ UStaticMesh* UnFbx::FFbxImporter::ImportStaticMeshAsSingle(UObject* InParent, TA
 
 		if (Node->GetMesh())
 		{
-			if (!BuildStaticMeshFromGeometry(Node, StaticMesh, MeshMaterials, LODIndex, NewRawMesh,
+			if (!BuildStaticMeshFromGeometry(Node, StaticMesh, MeshMaterials, LODIndex,
 											 VertexColorImportOption, ExistingVertexColorData, ImportOptions->VertexOverrideColor))
 			{
 				bBuildStatus = false;
@@ -1630,7 +1620,7 @@ UStaticMesh* UnFbx::FFbxImporter::ImportStaticMeshAsSingle(UObject* InParent, TA
 	if (bBuildStatus)
 	{
 		TVertexInstanceAttributesRef<FVector2D> VertexInstanceUVs = MeshDescription->VertexInstanceAttributes().GetAttributesRef<FVector2D>(MeshAttribute::VertexInstance::TextureCoordinate);
-		int32 FirstOpenUVChannel = VertexInstanceUVs.GetNumIndices() >= MAX_MESH_TEXTURE_COORDS ? 1 : VertexInstanceUVs.GetNumIndices();
+		int32 FirstOpenUVChannel = VertexInstanceUVs.GetNumIndices() >= MAX_MESH_TEXTURE_COORDS_MD ? 1 : VertexInstanceUVs.GetNumIndices();
 		TPolygonGroupAttributesRef<FName> PolygonGroupImportedMaterialSlotNames = MeshDescription->PolygonGroupAttributes().GetAttributesRef<FName>(MeshAttribute::PolygonGroup::ImportedMaterialSlotName);
 
 		TArray<FStaticMaterial> MaterialToAdd;
diff --git a/Engine/Source/Editor/UnrealEd/Private/StaticMeshEdit.cpp b/Engine/Source/Editor/UnrealEd/Private/StaticMeshEdit.cpp
--- a/Engine/Source/Editor/UnrealEd/Private/StaticMeshEdit.cpp
+++ b/Engine/Source/Editor/UnrealEd/Private/StaticMeshEdit.cpp
@@ -585,7 +585,7 @@ void AddConvexGeomFromVertices( const TArray<FVector>& Verts, FKAggregateGeom* A
 /**
  * Creates a static mesh object from raw triangle data.
  */
-UStaticMesh* CreateStaticMesh(struct FRawMesh& RawMesh,TArray<FStaticMaterial>& Materials,UObject* InOuter,FName InName)
+UStaticMesh* CreateStaticMesh(FMeshDescription& RawMesh,TArray<FStaticMaterial>& Materials,UObject* InOuter,FName InName)
 {
 	// Create the UStaticMesh object.
 	FStaticMeshComponentRecreateRenderStateContext RecreateRenderStateContext(FindObject<UStaticMesh>(InOuter,*InName.ToString()));
@@ -593,7 +593,9 @@ UStaticMesh* CreateStaticMesh(struct FRawMesh& RawMesh,TArray<FStaticMaterial>&
 
 	// Add one LOD for the base mesh
 	FStaticMeshSourceModel& SrcModel = StaticMesh->AddSourceModel();
-	SrcModel.SaveRawMesh(RawMesh);
+	FMeshDescription* MeshDescription = StaticMesh->CreateOriginalMeshDescription(0);
+	*MeshDescription = RawMesh;
+	StaticMesh->CommitOriginalMeshDescription(0);
 	StaticMesh->StaticMaterials = Materials;
 
 	int32 NumSections = StaticMesh->StaticMaterials.Num();
@@ -810,89 +812,6 @@ void GetBrushMesh(ABrush* Brush, UModel* Model, FMeshDescription& MeshDescriptio
 	}
 }
 
-void GetBrushMesh(ABrush* Brush,UModel* Model,struct FRawMesh& OutMesh,TArray<FStaticMaterial>& OutMaterials)
-{
-	// Calculate the local to world transform for the source brush.
-
-	FMatrix	ActorToWorld = Brush ? Brush->ActorToWorld().ToMatrixWithScale() : FMatrix::Identity;
-	bool	ReverseVertices = 0;
-	FVector4	PostSub = Brush ? FVector4(Brush->GetActorLocation()) : FVector4(0,0,0,0);
-
-
-	int32 NumPolys = Model->Polys->Element.Num();
-
-	// For each polygon in the model...
-	TArray<FVector> TempPositions;
-	for( int32 PolygonIndex = 0; PolygonIndex < NumPolys; ++PolygonIndex )
-	{
-		FPoly& Polygon = Model->Polys->Element[PolygonIndex];
-
-		UMaterialInterface*	Material = Polygon.Material;
-
-		// Find a material index for this polygon.
-
-		int32 MaterialIndex = OutMaterials.AddUnique(FStaticMaterial(Material));
-
-		// Cache the texture coordinate system for this polygon.
-
-		FVector	TextureBase = Polygon.Base - (Brush ? Brush->GetPivotOffset() : FVector::ZeroVector),
-				TextureX = Polygon.TextureU / UModel::GetGlobalBSPTexelScale(),
-				TextureY = Polygon.TextureV / UModel::GetGlobalBSPTexelScale();
-
-		// For each vertex after the first two vertices...
-		for(int32 VertexIndex = 2;VertexIndex < Polygon.Vertices.Num();VertexIndex++)
-		{
-			// Create a triangle for the vertex.
-			OutMesh.FaceMaterialIndices.Add(MaterialIndex);
-
-			// Generate different smoothing mask for each poly to give the mesh hard edges.  Note: only 32 smoothing masks supported.
-			OutMesh.FaceSmoothingMasks.Add(1<<(PolygonIndex%32));
-
-			FVector Positions[3];
-			FVector2D UVs[3];
-
-
-			Positions[ReverseVertices ? 0 : 2] = ActorToWorld.TransformPosition(Polygon.Vertices[0]) - PostSub;
-			UVs[ReverseVertices ? 0 : 2].X = (Positions[ReverseVertices ? 0 : 2] - TextureBase) | TextureX;
-			UVs[ReverseVertices ? 0 : 2].Y = (Positions[ReverseVertices ? 0 : 2] - TextureBase) | TextureY;
-
-			Positions[1] = ActorToWorld.TransformPosition(Polygon.Vertices[VertexIndex - 1]) - PostSub;
-			UVs[1].X = (Positions[1] - TextureBase) | TextureX;
-			UVs[1].Y = (Positions[1] - TextureBase) | TextureY;
-
-			Positions[ReverseVertices ? 2 : 0] = ActorToWorld.TransformPosition(Polygon.Vertices[VertexIndex]) - PostSub;
-			UVs[ReverseVertices ? 2 : 0].X = (Positions[ReverseVertices ? 2 : 0] - TextureBase) | TextureX;
-			UVs[ReverseVertices ? 2 : 0].Y = (Positions[ReverseVertices ? 2 : 0] - TextureBase) | TextureY;
-
-			for (int32 CornerIndex = 0; CornerIndex < 3; ++CornerIndex)
-			{
-				TempPositions.Add(Positions[CornerIndex]);
-				OutMesh.WedgeTexCoords[0].Add(UVs[CornerIndex]);
-			}
-		}
-	}
-
-	// Merge vertices within a certain distance of each other.
-	for (int32 i = 0; i < TempPositions.Num(); ++i)
-	{
-		FVector Position = TempPositions[i];
-		int32 FinalIndex = INDEX_NONE;
-		for (int32 VertexIndex = 0; VertexIndex < OutMesh.VertexPositions.Num(); ++VertexIndex)
-		{
-			if (FVerticesEqual(Position, OutMesh.VertexPositions[VertexIndex]))
-			{
-				FinalIndex = VertexIndex;
-				break;
-			}
-		}
-		if (FinalIndex == INDEX_NONE)
-		{
-			FinalIndex = OutMesh.VertexPositions.Add(Position);
-		}
-		OutMesh.WedgeIndices.Add(FinalIndex);
-	}
-}
-
 //
 //	CreateStaticMeshFromBrush - Creates a static mesh from the triangles in a model.
 //
diff --git a/Engine/Source/Editor/UnrealEd/Public/Editor.h b/Engine/Source/Editor/UnrealEd/Public/Editor.h
--- a/Engine/Source/Editor/UnrealEd/Public/Editor.h
+++ b/Engine/Source/Editor/UnrealEd/Public/Editor.h
@@ -419,7 +419,7 @@ extern bool GBuildStaticMeshCollision;
 //
 // Creating a static mesh from an array of triangles.
 //
-UStaticMesh* CreateStaticMesh(struct FRawMesh& RawMesh,TArray<FStaticMaterial>& Materials,UObject* Outer,FName Name);
+UStaticMesh* CreateStaticMesh(struct FMeshDescription& RawMesh,TArray<FStaticMaterial>& Materials,UObject* Outer,FName Name);
 
 struct FMergeStaticMeshParams
 {
@@ -466,7 +466,6 @@ void MergeStaticMesh(UStaticMesh* DestMesh, UStaticMesh* SourceMesh, const FMerg
 //
 // Converting models to static meshes.
 //
-UNREALED_API void GetBrushMesh(ABrush* Brush,UModel* Model,struct FRawMesh& OutMesh,TArray<FStaticMaterial>& OutMaterials);
 UNREALED_API void GetBrushMesh(ABrush* Brush, UModel* Model, struct FMeshDescription& OutMesh, TArray<FStaticMaterial>& OutMaterials);
 UStaticMesh* CreateStaticMeshFromBrush(UObject* Outer,FName Name,ABrush* Brush,UModel* Model);
  
diff --git a/Engine/Source/Editor/UnrealEd/Public/FbxImporter.h b/Engine/Source/Editor/UnrealEd/Public/FbxImporter.h
--- a/Engine/Source/Editor/UnrealEd/Public/FbxImporter.h
+++ b/Engine/Source/Editor/UnrealEd/Public/FbxImporter.h
@@ -1294,7 +1294,7 @@ protected:
 	 * @param LODIndex	LOD level to set up for StaticMesh
 	 * @return bool true if set up successfully
 	 */
-	bool BuildStaticMeshFromGeometry(FbxNode* Node, UStaticMesh* StaticMesh, TArray<FFbxMaterial>& MeshMaterials, int LODIndex, FRawMesh& RawMesh,
+	bool BuildStaticMeshFromGeometry(FbxNode* Node, UStaticMesh* StaticMesh, TArray<FFbxMaterial>& MeshMaterials, int LODIndex,
 									 EVertexColorImportOption::Type VertexColorImportOption, const TMap<FVector, FColor>& ExistingVertexColorData, const FColor& VertexOverrideColor);
 	
 	/**
diff --git a/Engine/Source/Editor/UnrealEd/UnrealEd.Build.cs b/Engine/Source/Editor/UnrealEd/UnrealEd.Build.cs
--- a/Engine/Source/Editor/UnrealEd/UnrealEd.Build.cs
+++ b/Engine/Source/Editor/UnrealEd/UnrealEd.Build.cs
@@ -96,6 +96,7 @@ public class UnrealEd : ModuleRules
 				"UMG",
 				"NavigationSystem",
                 "MeshDescription",
+                "MeshDescriptionOperations",
                 "MeshBuilder",
                 "MaterialShaderQualitySettings"
             }
