Add the possibility to reduce static mesh LOD from any lower LOD instead of reducing from LOD 0 only.

diff --git a/Engine/Source/Developer/MeshBuilder/Private/StaticMeshBuilder.cpp b/Engine/Source/Developer/MeshBuilder/Private/StaticMeshBuilder.cpp
--- a/Engine/Source/Developer/MeshBuilder/Private/StaticMeshBuilder.cpp
+++ b/Engine/Source/Developer/MeshBuilder/Private/StaticMeshBuilder.cpp
@@ -83,7 +83,6 @@ bool FStaticMeshBuilder::Build(FStaticMeshRenderData& StaticMeshRenderData, USta
 
 	for (int32 LodIndex = 0; LodIndex < StaticMesh->SourceModels.Num(); ++LodIndex)
 	{
-		const int32 BaseReduceLodIndex = 0;
 		float MaxDeviation = 0.0f;
 		FMeshBuildSettings& LODBuildSettings = StaticMesh->SourceModels[LodIndex].BuildSettings;
 		const FMeshDescription* OriginalMeshDescription = StaticMesh->GetOriginalMeshDescription(LodIndex);
@@ -92,6 +91,7 @@ bool FStaticMeshBuilder::Build(FStaticMeshRenderData& StaticMeshRenderData, USta
 		const FStaticMeshSourceModel& SrcModel = StaticMesh->SourceModels[LodIndex];
 		FMeshReductionSettings ReductionSettings = LODGroup.GetSettings(SrcModel.ReductionSettings, LodIndex);
 
+		const int32 BaseReduceLodIndex = ReductionSettings.BaseLODModel;
 		// Use simplifier if a reduction in triangles or verts has been requested.
 
 		const bool bVertTermination = (!bIsThirdPartyReductiontool) && (ReductionSettings.TerminationCriterion != EStaticMeshReductionTerimationCriterion::Triangles);
@@ -150,7 +150,7 @@ bool FStaticMeshBuilder::Build(FStaticMeshRenderData& StaticMeshRenderData, USta
 		//Reduce LODs
 		if (bUseReduction)
 		{
-			const int32 BaseLodIndex = 0;
+			const int32 BaseLodIndex = ReductionSettings.BaseLODModel;
 			float OverlappingThreshold = LODBuildSettings.bRemoveDegenerates ? THRESH_POINTS_ARE_SAME : 0.0f;
 			FOverlappingCorners OverlappingCorners;
 			FMeshDescriptionOperations::FindOverlappingCorners(OverlappingCorners, MeshDescriptions[BaseLodIndex], OverlappingThreshold);
diff --git a/Engine/Source/Editor/StaticMeshEditor/Private/StaticMeshEditorTools.cpp b/Engine/Source/Editor/StaticMeshEditor/Private/StaticMeshEditorTools.cpp
--- a/Engine/Source/Editor/StaticMeshEditor/Private/StaticMeshEditorTools.cpp
+++ b/Engine/Source/Editor/StaticMeshEditor/Private/StaticMeshEditorTools.cpp
@@ -1,4 +1,4 @@
-﻿// Copyright 1998-2018 Epic Games, Inc. All Rights Reserved.
+// Copyright 1998-2018 Epic Games, Inc. All Rights Reserved.
 
 #include "StaticMeshEditorTools.h"
 #include "Framework/Commands/UIAction.h"
@@ -26,6 +26,7 @@
 #include "PhysicsEngine/BodySetup.h"
 #include "FbxMeshUtils.h"
 #include "Widgets/Input/SVectorInputBox.h"
+#include "Widgets/Input/SNumericEntryBox.h"
 #include "SPerPlatformPropertiesWidget.h"
 
 #include "Runtime/Analytics/Analytics/Public/Interfaces/IAnalyticsProvider.h"
@@ -1011,8 +1012,9 @@ void FMeshBuildSettingsLayout::OnDistanceFieldResolutionScaleCommitted(float New
 	OnDistanceFieldResolutionScaleChanged(NewValue);
 }
 
-FMeshReductionSettingsLayout::FMeshReductionSettingsLayout( TSharedRef<FLevelOfDetailSettingsLayout> InParentLODSettings )
+FMeshReductionSettingsLayout::FMeshReductionSettingsLayout( TSharedRef<FLevelOfDetailSettingsLayout> InParentLODSettings, int32 InCurrentLODIndex)
 	: ParentLODSettings( InParentLODSettings )
+	, CurrentLODIndex( InCurrentLODIndex )
 {
 
 	FillEnumOptions(ImportanceOptions, GetFeatureImportanceEnum());
@@ -1319,6 +1321,31 @@ void FMeshReductionSettingsLayout::GenerateChildContent( IDetailChildrenBuilder&
 
 		}
 	}
+
+	//Base LOD
+	{
+		ChildrenBuilder.AddCustomRow( LOCTEXT("BaseLOD", "Base LOD") )
+			.NameContent()
+			.HAlign(HAlign_Left)
+			[
+				SNew(STextBlock)
+				.Text(LOCTEXT("BaseLOD", "Base LOD"))
+				.Font(IDetailLayoutBuilder::GetDetailFont())
+			]
+			.ValueContent()
+			.HAlign(HAlign_Left)
+			[
+				SNew(SNumericEntryBox<int32>)
+				.AllowSpin(true)
+				.MinSliderValue(0)
+				.MaxSliderValue(CurrentLODIndex)
+				.MinValue(0)
+				.MaxValue(CurrentLODIndex)
+				.Value(this, &FMeshReductionSettingsLayout::GetBaseLODIndex)
+				.OnValueChanged(this, &FMeshReductionSettingsLayout::SetBaseLODIndex)
+			];
+	}
+
 	{
 		ChildrenBuilder.AddCustomRow( LOCTEXT("ApplyChanges", "Apply Changes") )
 			.ValueContent()
@@ -1554,6 +1581,19 @@ void FMeshReductionSettingsLayout::OnTerminationCriterionChanged(TSharedPtr<FStr
 	}
 }
 
+TOptional<int32> FMeshReductionSettingsLayout::GetBaseLODIndex() const
+{
+	return ReductionSettings.BaseLODModel;
+}
+
+void FMeshReductionSettingsLayout::SetBaseLODIndex(int32 NewLODBaseIndex)
+{
+	if (NewLODBaseIndex <= CurrentLODIndex)
+	{
+		ReductionSettings.BaseLODModel = NewLODBaseIndex;
+	}
+}
+
 FMeshSectionSettingsLayout::~FMeshSectionSettingsLayout()
 {
 }
@@ -3178,7 +3218,7 @@ void FLevelOfDetailSettingsLayout::AddLODLevelCategories( IDetailLayoutBuilder&
 
 			if (IsAutoMeshReductionAvailable())
 			{
-				ReductionSettingsWidgets[LODIndex] = MakeShareable( new FMeshReductionSettingsLayout( AsShared() ) );
+				ReductionSettingsWidgets[LODIndex] = MakeShareable( new FMeshReductionSettingsLayout( AsShared(), LODIndex) );
 			}
 
 			if (LODIndex < StaticMesh->SourceModels.Num())
diff --git a/Engine/Source/Editor/StaticMeshEditor/Private/StaticMeshEditorTools.h b/Engine/Source/Editor/StaticMeshEditor/Private/StaticMeshEditorTools.h
--- a/Engine/Source/Editor/StaticMeshEditor/Private/StaticMeshEditorTools.h
+++ b/Engine/Source/Editor/StaticMeshEditor/Private/StaticMeshEditorTools.h
@@ -215,7 +215,7 @@ private:
 class FMeshReductionSettingsLayout : public IDetailCustomNodeBuilder, public TSharedFromThis<FMeshReductionSettingsLayout>
 {
 public:
-	FMeshReductionSettingsLayout( TSharedRef<FLevelOfDetailSettingsLayout> InParentLODSettings );
+	FMeshReductionSettingsLayout( TSharedRef<FLevelOfDetailSettingsLayout> InParentLODSettings, int32 InCurrentLODIndex);
 	virtual ~FMeshReductionSettingsLayout();
 
 	const FMeshReductionSettings& GetSettings() const;
@@ -282,9 +282,13 @@ private:
 	EVisibility GetTriangleCriterionVisibility() const;
 	EVisibility GetVertexCriterionVisibility() const;
 
+	TOptional<int32> GetBaseLODIndex() const;
+	void SetBaseLODIndex(int32 NewLODBaseIndex);
+
 private:
 	TWeakPtr<FLevelOfDetailSettingsLayout> ParentLODSettings;
 	FMeshReductionSettings ReductionSettings;
+	int32 CurrentLODIndex;
 
 	// Used by simplygon
 	TArray<TSharedPtr<FString> > ImportanceOptions;
