When adding generated lod in staticmesh editor, set the percent triangle depending on the lod index.
Also make sure custom LOD has the base LOD set to themself

diff --git a/Engine/Source/Editor/StaticMeshEditor/Private/StaticMeshEditorTools.cpp b/Engine/Source/Editor/StaticMeshEditor/Private/StaticMeshEditorTools.cpp
--- a/Engine/Source/Editor/StaticMeshEditor/Private/StaticMeshEditorTools.cpp
+++ b/Engine/Source/Editor/StaticMeshEditor/Private/StaticMeshEditorTools.cpp
@@ -1326,12 +1326,12 @@ void FMeshReductionSettingsLayout::GenerateChildContent( IDetailChildrenBuilder&
 	//Base LOD
 	{
 		int32 MaxBaseReduceIndex = bCanReduceMyself ? CurrentLODIndex : CurrentLODIndex - 1;
-		ChildrenBuilder.AddCustomRow( LOCTEXT("BaseLOD", "Base LOD") )
+		ChildrenBuilder.AddCustomRow( LOCTEXT("ReductionBaseLOD", "Base LOD") )
 			.NameContent()
 			.HAlign(HAlign_Left)
 			[
 				SNew(STextBlock)
-				.Text(LOCTEXT("BaseLOD", "Base LOD"))
+				.Text(LOCTEXT("ReductionBaseLOD", "Base LOD"))
 				.Font(IDetailLayoutBuilder::GetDetailFont())
 			]
 			.ValueContent()
@@ -3747,7 +3747,18 @@ void FLevelOfDetailSettingsLayout::OnImportLOD(TSharedPtr<FString> NewValue, ESe
 	{
 		UStaticMesh* StaticMesh = StaticMeshEditor.GetStaticMesh();
 		check(StaticMesh);
-		FbxMeshUtils::ImportMeshLODDialog(StaticMesh,LODIndex);
+
+		//Are we a new imported LOD, we want to set some value for new imported LOD.
+		//This boolean prevent changing the value when the LOD is reimport
+		bool bImportCustomLOD = (LODIndex >= StaticMesh->SourceModels.Num());
+		
+		bool bResult = FbxMeshUtils::ImportMeshLODDialog(StaticMesh,LODIndex);
+
+		if (bImportCustomLOD && bResult && StaticMesh->SourceModels.IsValidIndex(LODIndex))
+		{
+			//Custom LOD should reduce base on them self when they get imported.
+			StaticMesh->SourceModels[LODIndex].ReductionSettings.BaseLODModel = LODIndex;
+		}
 		StaticMesh->PostEditChange();
 		StaticMeshEditor.RefreshTool();
 	}
diff --git a/Engine/Source/Runtime/Engine/Private/StaticMesh.cpp b/Engine/Source/Runtime/Engine/Private/StaticMesh.cpp
--- a/Engine/Source/Runtime/Engine/Private/StaticMesh.cpp
+++ b/Engine/Source/Runtime/Engine/Private/StaticMesh.cpp
@@ -2330,6 +2330,11 @@ void UStaticMesh::SetNumSourceModels(const int32 Num)
 			}
 		}
 		SourceModels[Index].ReductionSettings.BaseLODModel = PreviousCustomLODIndex;
+		if (GetOriginalMeshDescription(Index) == nullptr && !IsReductionActive(Index))
+		{
+			//Set the Reduction percent
+			SourceModels[Index].ReductionSettings.PercentTriangles = FMath::Pow(0.5f, (float)(Index-PreviousCustomLODIndex));
+		}
 	}
 }
 
