Fix some UI string to remove grammatical issues

diff --git a/Engine/Source/Editor/Persona/Private/PersonaMeshDetails.cpp b/Engine/Source/Editor/Persona/Private/PersonaMeshDetails.cpp
--- a/Engine/Source/Editor/Persona/Private/PersonaMeshDetails.cpp
+++ b/Engine/Source/Editor/Persona/Private/PersonaMeshDetails.cpp
@@ -2726,7 +2726,7 @@ FText FPersonaMeshDetails::GetCurrentLodTooltip() const
 {
 	if (GetPersonaToolkit()->GetPreviewMeshComponent() != nullptr && GetPersonaToolkit()->GetPreviewMeshComponent()->ForcedLodModel == 0)
 	{
-		return FText::FromString(TEXT("LOD0 is edit when selecting Auto LOD"));
+		return LOCTEXT("PersonaLODPickerCurrentLODTooltip", "With Auto LOD selected, LOD0's properties are visible for editing");
 	}
 	return FText::GetEmpty();
 }
diff --git a/Engine/Source/Editor/StaticMeshEditor/Private/StaticMeshEditorTools.cpp b/Engine/Source/Editor/StaticMeshEditor/Private/StaticMeshEditorTools.cpp
--- a/Engine/Source/Editor/StaticMeshEditor/Private/StaticMeshEditorTools.cpp
+++ b/Engine/Source/Editor/StaticMeshEditor/Private/StaticMeshEditorTools.cpp
@@ -4170,7 +4170,7 @@ FText FLevelOfDetailSettingsLayout::GetCurrentLodTooltip() const
 {
 	if (StaticMeshEditor.GetStaticMeshComponent() != nullptr && StaticMeshEditor.GetStaticMeshComponent()->ForcedLodModel == 0)
 	{
-		return FText::FromString(TEXT("LOD0 is edit when selecting Auto LOD"));
+		return LOCTEXT("StaticMeshEditorLODPickerCurrentLODTooltip", "With Auto LOD selected, LOD0's properties are visible for editing");
 	}
 	return FText::GetEmpty();
 }
