Rename reimport menu display label

diff --git a/Engine/Source/Editor/SkeletalMeshEditor/Private/SkeletalMeshEditorCommands.cpp b/Engine/Source/Editor/SkeletalMeshEditor/Private/SkeletalMeshEditorCommands.cpp
--- a/Engine/Source/Editor/SkeletalMeshEditor/Private/SkeletalMeshEditorCommands.cpp
+++ b/Engine/Source/Editor/SkeletalMeshEditor/Private/SkeletalMeshEditorCommands.cpp
@@ -7,8 +7,8 @@
 
 void FSkeletalMeshEditorCommands::RegisterCommands()
 {
-	UI_COMMAND(ReimportMesh, "Reimport Mesh", "Reimport the current mesh.", EUserInterfaceActionType::Button, FInputChord());
-	UI_COMMAND(ReimportAllMesh, "Reimport All Mesh", "Reimport the current mesh and all the custom LODs.", EUserInterfaceActionType::Button, FInputChord());
+	UI_COMMAND(ReimportMesh, "Reimport Base Mesh", "Reimport the base mesh.", EUserInterfaceActionType::Button, FInputChord());
+	UI_COMMAND(ReimportAllMesh, "Reimport Base + LODs", "Reimport the base mesh and all the custom LODs.", EUserInterfaceActionType::Button, FInputChord());
 
 	UI_COMMAND(MeshSectionSelection, "Section Selection", "Enables selecting Mesh Sections in the viewport (disables selecting bones using their physics shape).", EUserInterfaceActionType::ToggleButton, FInputChord());
 }
diff --git a/Engine/Source/Editor/StaticMeshEditor/Private/StaticMeshEditorActions.cpp b/Engine/Source/Editor/StaticMeshEditor/Private/StaticMeshEditorActions.cpp
--- a/Engine/Source/Editor/StaticMeshEditor/Private/StaticMeshEditorActions.cpp
+++ b/Engine/Source/Editor/StaticMeshEditor/Private/StaticMeshEditorActions.cpp
@@ -17,8 +17,8 @@ void FStaticMeshEditorCommands::RegisterCommands()
 	UI_COMMAND( SetShowSockets, "Sockets", "Displays the static mesh sockets.", EUserInterfaceActionType::ToggleButton, FInputChord() );
 	UI_COMMAND( SetDrawAdditionalData, "Additional Data", "Draw additional user data associated with asset.", EUserInterfaceActionType::ToggleButton, FInputChord() );
 
-	UI_COMMAND( ReimportMesh, "Reimport Mesh", "Reimport the current mesh.", EUserInterfaceActionType::Button, FInputChord() );
-	UI_COMMAND( ReimportAllMesh, "Reimport All Mesh", "Reimport the current mesh and all the custom LODs.", EUserInterfaceActionType::Button, FInputChord() );
+	UI_COMMAND( ReimportMesh, "Reimport Base Mesh", "Reimport the base mesh.", EUserInterfaceActionType::Button, FInputChord() );
+	UI_COMMAND( ReimportAllMesh, "Reimport Base + LODs", "Reimport the base mesh and all the custom LODs.", EUserInterfaceActionType::Button, FInputChord() );
 
 	UI_COMMAND( SetShowNormals, "Normals", "Toggles display of vertex normals in the Preview Pane.", EUserInterfaceActionType::ToggleButton, FInputChord() );
 	UI_COMMAND( SetShowTangents, "Tangents", "Toggles display of vertex tangents in the Preview Pane.", EUserInterfaceActionType::ToggleButton, FInputChord() );
