Fix a potential data race when changing UStaticMesh::NumStreamedLODs through editor

diff --git a/Engine/Source/Editor/StaticMeshEditor/Private/StaticMeshEditorTools.cpp b/Engine/Source/Editor/StaticMeshEditor/Private/StaticMeshEditorTools.cpp
--- a/Engine/Source/Editor/StaticMeshEditor/Private/StaticMeshEditorTools.cpp
+++ b/Engine/Source/Editor/StaticMeshEditor/Private/StaticMeshEditorTools.cpp
@@ -43,6 +43,7 @@
 #include "Factories/FbxStaticMeshImportData.h"
 #include "Interfaces/ITargetPlatformManagerModule.h"
 #include "Interfaces/ITargetPlatform.h"
+#include "ContentStreaming.h"
 
 
 const uint32 MaxHullCount = 64;
@@ -4101,6 +4102,9 @@ void FLevelOfDetailSettingsLayout::OnNumStreamedLODsCommitted(int32 InValue, ETe
 	check(StaticMesh);
 	if (UpdateStaticMeshNumStreamedLODsHelper(StaticMesh, InValue, Platform))
 	{
+		// Make sure FStaticMeshRenderData::CurrentFirstLODIdx is not accessed on other threads
+		IStreamingManager::Get().GetTextureStreamingManager().BlockTillAllRequestsFinished();
+		// Recache derived data and relink streaming
 		ApplyChanges();
 	}
 	StaticMeshEditor.RefreshViewport();
