Fix for LOD radio button not updating static mesh editor viewport

diff --git a/Engine/Source/Editor/StaticMeshEditor/Private/SStaticMeshEditorViewport.cpp b/Engine/Source/Editor/StaticMeshEditor/Private/SStaticMeshEditorViewport.cpp
--- a/Engine/Source/Editor/StaticMeshEditor/Private/SStaticMeshEditorViewport.cpp
+++ b/Engine/Source/Editor/StaticMeshEditor/Private/SStaticMeshEditorViewport.cpp
@@ -412,6 +412,7 @@ void SStaticMeshEditorViewport::OnSetLODModel(int32 InLODSelection)
 		PreviewMeshComponent->SetForcedLodModel(LODSelection);
 		//PopulateUVChoices();
 		StaticMeshEditorPtr.Pin()->BroadcastOnSelectedLODChanged();
+		RefreshViewport();
 	}
 }
 
