Option to specify static mesh complex collision from a different asset.

diff --git a/Engine/Source/Editor/StaticMeshEditor/Private/StaticMeshEditorViewportClient.cpp b/Engine/Source/Editor/StaticMeshEditor/Private/StaticMeshEditorViewportClient.cpp
--- a/Engine/Source/Editor/StaticMeshEditor/Private/StaticMeshEditorViewportClient.cpp
+++ b/Engine/Source/Editor/StaticMeshEditor/Private/StaticMeshEditorViewportClient.cpp
@@ -439,6 +439,16 @@ bool FStaticMeshEditorViewportClient::ShouldOrbitCamera() const
 	return FEditorViewportClient::ShouldOrbitCamera();
 }
 
+void DrawCustomComplex(FPrimitiveDrawInterface* PDI, FTriMeshCollisionData Mesh, const FColor Color)
+{
+	for (int i = 0; i < Mesh.Indices.Num(); ++i)
+	{
+		PDI->DrawLine(Mesh.Vertices[Mesh.Indices[i].v0], Mesh.Vertices[Mesh.Indices[i].v1], Color, SDPG_World);
+		PDI->DrawLine(Mesh.Vertices[Mesh.Indices[i].v1], Mesh.Vertices[Mesh.Indices[i].v2], Color, SDPG_World);
+		PDI->DrawLine(Mesh.Vertices[Mesh.Indices[i].v2], Mesh.Vertices[Mesh.Indices[i].v0], Color, SDPG_World);
+	}
+}
+
 void FStaticMeshEditorViewportClient::Draw(const FSceneView* View,FPrimitiveDrawInterface* PDI)
 {
 	FEditorViewportClient::Draw(View, PDI);
@@ -518,6 +528,18 @@ void FStaticMeshEditorViewportClient::Draw(const FSceneView* View,FPrimitiveDraw
 		}
 	}
 
+	if (bShowComplexCollision && StaticMesh->ComplexCollisionMesh && StaticMesh->BodySetup->CollisionTraceFlag != ECollisionTraceFlag::CTF_UseSimpleAsComplex)
+	{
+		const FColor SelectedColor(20, 20, 220);
+		const FColor UnselectedColor(0, 0, 125);
+
+		HSMECollisionProxy* HitProxy = new HSMECollisionProxy(EAggCollisionShape::Convex, 0);
+		PDI->SetHitProxy(HitProxy);
+		const FColor CollisionColor = StaticMeshEditor->IsSelectedPrim(HitProxy->PrimData) ? SelectedColor : UnselectedColor;
+		DrawCustomComplex(PDI, CollisionMeshData, CollisionColor);
+		PDI->SetHitProxy(nullptr);
+	}
+
 	if( bShowSockets )
 	{
 		const FColor SocketColor = FColor(255, 128, 128);
@@ -1447,10 +1469,26 @@ void FStaticMeshEditorViewportClient::ToggleShowComplexCollision()
 {
 	bShowComplexCollision = !bShowComplexCollision;
 
-	if (StaticMeshComponent != nullptr)
+	if (StaticMesh)
 	{
-		StaticMeshComponent->bDrawMeshCollisionIfComplex = bShowComplexCollision;
-		StaticMeshComponent->MarkRenderStateDirty();
+		if (UObject* CDPObj = StaticMesh->ComplexCollisionMesh)
+		{
+			if (IInterface_CollisionDataProvider* CDP = Cast<IInterface_CollisionDataProvider>(CDPObj))
+			{
+				CollisionMeshData = FTriMeshCollisionData();
+				CDP->GetPhysicsTriMeshData(&CollisionMeshData, true);
+			}
+			if (StaticMeshComponent != nullptr)
+			{
+				StaticMeshComponent->bDrawMeshCollisionIfComplex = false;
+				StaticMeshComponent->MarkRenderStateDirty();
+			}
+		}
+		else if (StaticMeshComponent != nullptr)
+		{
+			StaticMeshComponent->bDrawMeshCollisionIfComplex = bShowComplexCollision;
+			StaticMeshComponent->MarkRenderStateDirty();
+		}
 	}
 
 	if (FEngineAnalytics::IsAvailable())
@@ -1534,4 +1572,4 @@ void FStaticMeshEditorViewportClient::ResetCamera()
 	FocusViewportOnBox( StaticMeshComponent->Bounds.GetBox() );
 	Invalidate();
 }
-#undef LOCTEXT_NAMESPACE 
+#undef LOCTEXT_NAMESPACE 
\ No newline at end of file
diff --git a/Engine/Source/Editor/StaticMeshEditor/Private/StaticMeshEditorViewportClient.h b/Engine/Source/Editor/StaticMeshEditor/Private/StaticMeshEditorViewportClient.h
--- a/Engine/Source/Editor/StaticMeshEditor/Private/StaticMeshEditorViewportClient.h
+++ b/Engine/Source/Editor/StaticMeshEditor/Private/StaticMeshEditorViewportClient.h
@@ -154,6 +154,8 @@ private:
 	/** The static mesh being used in the editor. */
 	UStaticMesh* StaticMesh;
 
+	FTriMeshCollisionData CollisionMeshData;
+
 	/** Pointer back to the StaticMesh editor tool that owns us */
 	TWeakPtr<IStaticMeshEditor> StaticMeshEditorPtr;
 
diff --git a/Engine/Source/Runtime/Engine/Classes/Engine/StaticMesh.h b/Engine/Source/Runtime/Engine/Classes/Engine/StaticMesh.h
--- a/Engine/Source/Runtime/Engine/Classes/Engine/StaticMesh.h
+++ b/Engine/Source/Runtime/Engine/Classes/Engine/StaticMesh.h
@@ -799,6 +799,8 @@ public:
 	UPROPERTY(Instanced, VisibleAnywhere, Category = EditableMesh)
 	class UObject* EditableMesh;
 
+	UPROPERTY(EditAnywhere, Category = Collision)
+	class UStaticMesh* ComplexCollisionMesh;
 	/**
 	 * Registers the mesh attributes required by the mesh description for a static mesh.
 	 */
diff --git a/Engine/Source/Runtime/Engine/Private/StaticMesh.cpp b/Engine/Source/Runtime/Engine/Private/StaticMesh.cpp
--- a/Engine/Source/Runtime/Engine/Private/StaticMesh.cpp
+++ b/Engine/Source/Runtime/Engine/Private/StaticMesh.cpp
@@ -2817,6 +2817,14 @@ void UStaticMesh::PostEditChangeProperty(FPropertyChangedEvent& PropertyChangedE
 		bool bRebuild = false;
 		SetLODGroup(LODGroup, bRebuild);
 	}
+	if (PropertyName == GET_MEMBER_NAME_CHECKED(UStaticMesh, ComplexCollisionMesh) && ComplexCollisionMesh != this)
+	{
+		if (BodySetup)
+		{
+			BodySetup->InvalidatePhysicsData();
+			BodySetup->CreatePhysicsMeshes();
+		}
+	}
 	LightMapResolution = FMath::Max(LightMapResolution, 0);
 
 	if (PropertyChangedEvent.MemberProperty 
@@ -5023,6 +5031,10 @@ static int32 GetCollisionVertIndexForMeshVertIndex(int32 MeshVertIndex, TMap<int
 bool UStaticMesh::GetPhysicsTriMeshData(struct FTriMeshCollisionData* CollisionData, bool bInUseAllTriData)
 {
 #if WITH_EDITORONLY_DATA
+	if (ComplexCollisionMesh)
+	{
+		return ComplexCollisionMesh->GetPhysicsTriMeshData(CollisionData, bInUseAllTriData);
+	}
 	check(HasValidRenderData());
 
 	// Get the LOD level to use for collision
@@ -5075,6 +5087,10 @@ bool UStaticMesh::GetPhysicsTriMeshData(struct FTriMeshCollisionData* CollisionD
 bool UStaticMesh::ContainsPhysicsTriMeshData(bool bInUseAllTriData) const 
 {
 #if WITH_EDITORONLY_DATA
+	if (ComplexCollisionMesh)
+	{
+		ComplexCollisionMesh->ContainsPhysicsTriMeshData(bInUseAllTriData);
+	}
 	if(RenderData == nullptr || RenderData->LODResources.Num() == 0)
 	{
 		return false;
