Fix Build (missing header reported in //UE4/main)

#rnx

diff --git a/Engine/Source/Editor/StaticMeshEditor/Private/StaticMeshEditorViewportClient.h b/Engine/Source/Editor/StaticMeshEditor/Private/StaticMeshEditorViewportClient.h
--- a/Engine/Source/Editor/StaticMeshEditor/Private/StaticMeshEditorViewportClient.h
+++ b/Engine/Source/Editor/StaticMeshEditor/Private/StaticMeshEditorViewportClient.h
@@ -7,6 +7,7 @@
 #include "UnrealWidget.h"
 #include "EditorViewportClient.h"
 #include "Components.h"
+#include "Interfaces/Interface_CollisionDataProvider.h"
 
 class FAdvancedPreviewScene;
 class FCanvas;
