Added missing property handle information
#rnx

diff --git a/Engine/Source/Editor/StaticMeshEditor/Private/StaticMeshEditorTools.cpp b/Engine/Source/Editor/StaticMeshEditor/Private/StaticMeshEditorTools.cpp
--- a/Engine/Source/Editor/StaticMeshEditor/Private/StaticMeshEditorTools.cpp
+++ b/Engine/Source/Editor/StaticMeshEditor/Private/StaticMeshEditorTools.cpp
@@ -2310,7 +2310,7 @@ UStaticMesh& FMeshMaterialsLayout::GetStaticMesh() const
 	return *StaticMesh;
 }
 
-void FMeshMaterialsLayout::AddToCategory(IDetailCategoryBuilder& CategoryBuilder)
+void FMeshMaterialsLayout::AddToCategory(IDetailCategoryBuilder& CategoryBuilder, TSharedPtr<class IPropertyHandle> InHandle)
 {
 	CategoryBuilder.AddCustomRow(LOCTEXT("AddLODLevelCategories_MaterialArrayOperationAdd", "Add Material Slot"))
 		.CopyAction(FUIAction(FExecuteAction::CreateSP(this, &FMeshMaterialsLayout::OnCopyMaterialList), FCanExecuteAction::CreateSP(this, &FMeshMaterialsLayout::OnCanCopyMaterialList)))
@@ -3268,9 +3268,9 @@ void FLevelOfDetailSettingsLayout::AddLODLevelCategories( IDetailLayoutBuilder&
 			FString CategoryName = FString(TEXT("StaticMeshMaterials"));
 
 			IDetailCategoryBuilder& MaterialsCategory = DetailBuilder.EditCategory(*CategoryName, LOCTEXT("StaticMeshMaterialsLabel", "Material Slots"), ECategoryPriority::Important);
-
+			TSharedPtr<IPropertyHandle> MeshLODHandle = DetailBuilder.GetProperty(GET_MEMBER_NAME_CHECKED(UStaticMesh, LODGroup));
 			MaterialsLayoutWidget = MakeShareable(new FMeshMaterialsLayout(StaticMeshEditor));
-			MaterialsLayoutWidget->AddToCategory(MaterialsCategory);
+			MaterialsLayoutWidget->AddToCategory(MaterialsCategory, MeshLODHandle);
 		}
 
 		int32 CurrentLodIndex = 0;
diff --git a/Engine/Source/Editor/StaticMeshEditor/Private/StaticMeshEditorTools.h b/Engine/Source/Editor/StaticMeshEditor/Private/StaticMeshEditorTools.h
--- a/Engine/Source/Editor/StaticMeshEditor/Private/StaticMeshEditorTools.h
+++ b/Engine/Source/Editor/StaticMeshEditor/Private/StaticMeshEditorTools.h
@@ -433,7 +433,7 @@ public:
 
 	virtual ~FMeshMaterialsLayout();
 
-	void AddToCategory(IDetailCategoryBuilder& CategoryBuilder);
+	void AddToCategory(IDetailCategoryBuilder& CategoryBuilder, TSharedPtr<class IPropertyHandle> InHandle);
 
 private:
 	UStaticMesh& GetStaticMesh() const;
