Path Selection Window
When user selects target path for exporting or saving, we need to refresh content from server. So we add OnPathSelected event to FSharedAssetDialogConfig and SAssetDialog.
Changed files: IContentBrowserSingleton.h, SAssetDialog.h, SAssetDialog.cpp.

GPumpingMessages
When files and folders are deleted or created by user, we want to save these changes to server. So we monitor such changes through IAssetRegistry. But we need to determine whether it∩┐╜s triggered by user. So we add GPumpingMessages, a global flag indicating that application is processing Windows messages. Which means it∩┐╜s triggered by user input.
Another possible choice is adding delegates. But this requires more changes and is less general.
Changed files: CoreGlobals.h, CoreGlobals.cpp, WindowsPlatformApplicationMisc.cpp.

Mesh Editor Small Fix
Fix an issue where vertex colors in Toolbar didn't work in 4.22.
Changed file is SStaticMeshEditorViewport.cpp.

PR #6121: Editor changes for NVIDIA USD plugin

diff --git a/Engine/Source/Editor/ContentBrowser/Private/SAssetDialog.cpp b/Engine/Source/Editor/ContentBrowser/Private/SAssetDialog.cpp
--- a/Engine/Source/Editor/ContentBrowser/Private/SAssetDialog.cpp
+++ b/Engine/Source/Editor/ContentBrowser/Private/SAssetDialog.cpp
@@ -90,6 +90,8 @@ void SAssetDialog::Construct(const FArguments& InArgs, const FSharedAssetDialogC
 	AssetPickerConfig.OnGetAssetContextMenu = FOnGetAssetContextMenu::CreateSP(this, &SAssetDialog::OnGetAssetContextMenu);
 	AssetPickerConfig.OnGetFolderContextMenu = FOnGetFolderContextMenu::CreateSP(this, &SAssetDialog::OnGetFolderContextMenu);
 
+	OnPathSelected = InConfig.OnPathSelected;
+
 	SetCurrentlySelectedPath(DefaultPath);
 
 	// Open and save specific configuration
@@ -818,6 +820,8 @@ void SAssetDialog::SetCurrentlySelectedPath(const FString& NewPath)
 {
 	CurrentlySelectedPath = NewPath;
 	UpdateInputValidity();
+
+	OnPathSelected.ExecuteIfBound(NewPath);
 }
 
 void SAssetDialog::SetCurrentlyEnteredAssetName(const FString& NewName)
diff --git a/Engine/Source/Editor/ContentBrowser/Private/SAssetDialog.h b/Engine/Source/Editor/ContentBrowser/Private/SAssetDialog.h
--- a/Engine/Source/Editor/ContentBrowser/Private/SAssetDialog.h
+++ b/Engine/Source/Editor/ContentBrowser/Private/SAssetDialog.h
@@ -208,4 +208,7 @@ private:
 
 	/** Utility member to know if the context menu was opened on the asset view or the path view */
 	EOpenedContextMenuWidget OpenedContextMenuWidget;
+
+	/** Callback for refreshing content */
+	FOnPathSelected OnPathSelected;
 };
diff --git a/Engine/Source/Editor/ContentBrowser/Public/IContentBrowserSingleton.h b/Engine/Source/Editor/ContentBrowser/Public/IContentBrowserSingleton.h
--- a/Engine/Source/Editor/ContentBrowser/Public/IContentBrowserSingleton.h
+++ b/Engine/Source/Editor/ContentBrowser/Public/IContentBrowserSingleton.h
@@ -381,6 +381,7 @@ struct FSharedAssetDialogConfig
 	FString DefaultPath;
 	TArray<FName> AssetClassNames;
 	FVector2D WindowSizeOverride;
+	FOnPathSelected OnPathSelected;
 
 	virtual EAssetDialogType::Type GetDialogType() const = 0;
 
diff --git a/Engine/Source/Editor/StaticMeshEditor/Private/SStaticMeshEditorViewport.cpp b/Engine/Source/Editor/StaticMeshEditor/Private/SStaticMeshEditorViewport.cpp
--- a/Engine/Source/Editor/StaticMeshEditor/Private/SStaticMeshEditorViewport.cpp
+++ b/Engine/Source/Editor/StaticMeshEditor/Private/SStaticMeshEditorViewport.cpp
@@ -349,9 +349,9 @@ void SStaticMeshEditorViewport::SetViewModeVertexColor()
 	if (!EditorViewportClient->EngineShowFlags.VertexColors)
 	{
 		EditorViewportClient->EngineShowFlags.SetVertexColors(true);
-		EditorViewportClient->EngineShowFlags.SetLighting(false);
 		EditorViewportClient->EngineShowFlags.SetIndirectLightingCache(false);
-		EditorViewportClient->EngineShowFlags.SetPostProcessing(false);
+		EditorViewportClient->EngineShowFlags.SetSelectionOutline(false);
+
 		EditorViewportClient->SetFloorAndEnvironmentVisibility(false);
 		StaticMeshEditorPtr.Pin()->GetStaticMeshComponent()->bDisplayVertexColors = true;
 		StaticMeshEditorPtr.Pin()->GetStaticMeshComponent()->MarkRenderStateDirty();
@@ -361,7 +361,7 @@ void SStaticMeshEditorViewport::SetViewModeVertexColor()
 		EditorViewportClient->EngineShowFlags.SetVertexColors(false);
 		EditorViewportClient->EngineShowFlags.SetLighting(true);
 		EditorViewportClient->EngineShowFlags.SetIndirectLightingCache(true);
-		EditorViewportClient->EngineShowFlags.SetPostProcessing(true);
+		EditorViewportClient->EngineShowFlags.SetSelectionOutline(GetDefault<ULevelEditorViewportSettings>()->bUseSelectionOutline);
 		EditorViewportClient->SetFloorAndEnvironmentVisibility(true);
 		StaticMeshEditorPtr.Pin()->GetStaticMeshComponent()->bDisplayVertexColors = false;
 		StaticMeshEditorPtr.Pin()->GetStaticMeshComponent()->MarkRenderStateDirty();
diff --git a/Engine/Source/Runtime/ApplicationCore/Private/Windows/WindowsPlatformApplicationMisc.cpp b/Engine/Source/Runtime/ApplicationCore/Private/Windows/WindowsPlatformApplicationMisc.cpp
--- a/Engine/Source/Runtime/ApplicationCore/Private/Windows/WindowsPlatformApplicationMisc.cpp
+++ b/Engine/Source/Runtime/ApplicationCore/Private/Windows/WindowsPlatformApplicationMisc.cpp
@@ -113,6 +113,13 @@ static void WinPumpMessages()
 
 void FWindowsPlatformApplicationMisc::PumpMessages(bool bFromMainLoop)
 {
+	TSharedPtr<void> RevertGlobalFlag;
+	if (!GPumpingMessages)
+	{
+		GPumpingMessages = true;
+		RevertGlobalFlag = MakeShareable<void>(nullptr, [](auto) {GPumpingMessages = false; });
+	}
+
 	if (!bFromMainLoop)
 	{
 		FPlatformMisc::PumpMessagesOutsideMainLoop();
diff --git a/Engine/Source/Runtime/Core/Private/Misc/CoreGlobals.cpp b/Engine/Source/Runtime/Core/Private/Misc/CoreGlobals.cpp
--- a/Engine/Source/Runtime/Core/Private/Misc/CoreGlobals.cpp
+++ b/Engine/Source/Runtime/Core/Private/Misc/CoreGlobals.cpp
@@ -304,6 +304,8 @@ bool					GIsDemoMode						= false;
 bool					GIsAutomationTesting					= false;
 /** Whether or not messages are being pumped outside of the main loop										*/
 bool					GPumpingMessagesOutsideOfMainLoop = false;
+/** Whether or not messages are being pumped */
+bool					GPumpingMessages = false;
 
 /** Enables various editor and HMD hacks that allow the experimental VR editor feature to work, perhaps at the expense of other systems */
 bool					GEnableVREditorHacks = false;
diff --git a/Engine/Source/Runtime/Core/Public/CoreGlobals.h b/Engine/Source/Runtime/Core/Public/CoreGlobals.h
--- a/Engine/Source/Runtime/Core/Public/CoreGlobals.h
+++ b/Engine/Source/Runtime/Core/Public/CoreGlobals.h
@@ -407,6 +407,9 @@ extern CORE_API bool GIsAutomationTesting;
 /** Whether or not messages are being pumped outside of main loop */
 extern CORE_API bool GPumpingMessagesOutsideOfMainLoop;
 
+/** Whether or not messages are being pumped */
+extern CORE_API bool GPumpingMessages;
+ 
 /** Enables various editor and HMD hacks that allow the experimental VR editor feature to work, perhaps at the expense of other systems */
 extern CORE_API bool GEnableVREditorHacks;
 
