Made UStaticMesh::NavCollision a transactable object when entering the static mesh editor so that Undo is correctly handled.

diff --git a/Engine/Source/Editor/StaticMeshEditor/Private/StaticMeshEditor.cpp b/Engine/Source/Editor/StaticMeshEditor/Private/StaticMeshEditor.cpp
--- a/Engine/Source/Editor/StaticMeshEditor/Private/StaticMeshEditor.cpp
+++ b/Engine/Source/Editor/StaticMeshEditor/Private/StaticMeshEditor.cpp
@@ -15,6 +15,7 @@
 #include "Editor.h"
 #include "UnrealEdGlobals.h"
 #include "StaticMeshEditorModule.h"
+#include "AI/Navigation/NavCollisionBase.h"
 
 #include "SStaticMeshEditorViewport.h"
 #include "PropertyEditorModule.h"
@@ -154,6 +155,10 @@ void FStaticMeshEditor::InitEditorForStaticMesh(UStaticMesh* ObjectToEdit)
 
 	// Support undo/redo
 	ObjectToEdit->SetFlags( RF_Transactional );
+	if (ObjectToEdit->NavCollision)
+	{
+		ObjectToEdit->NavCollision->SetFlags(RF_Transactional);
+	}
 
 	GEditor->RegisterForUndo( this );
 
