[NO NON-PII COMMIT MESSAGE]

diff --git a/Engine/Source/Editor/StaticMeshEditor/Private/StaticMeshEditor.cpp b/Engine/Source/Editor/StaticMeshEditor/Private/StaticMeshEditor.cpp
--- a/Engine/Source/Editor/StaticMeshEditor/Private/StaticMeshEditor.cpp
+++ b/Engine/Source/Editor/StaticMeshEditor/Private/StaticMeshEditor.cpp
@@ -2245,8 +2245,10 @@ TStatId FStaticMeshEditor::GetStatId() const
 
 bool FStaticMeshEditor::CanRemoveUVChannel()
 {
-	// Can remove UV channel if there's one that is currently being selected and displayed
-	return Viewport->GetViewportClient().IsDrawUVOverlayChecked();
+	// Can remove UV channel if there's one that is currently being selected and displayed, 
+	// and the current LOD has more than one UV channel
+	return Viewport->GetViewportClient().IsDrawUVOverlayChecked() && 
+		StaticMesh->GetNumUVChannels(GetCurrentLODIndex()) > 1;
 }
 
 void FStaticMeshEditor::RemoveCurrentUVChannel()
