Fixing VertexColor visualisation in StaticMeshEditor. Reverting some changes made in CL# 8042071 that removed the PostProcess/Lighting disabling in VertexColor mode.

diff --git a/Engine/Source/Editor/StaticMeshEditor/Private/SStaticMeshEditorViewport.cpp b/Engine/Source/Editor/StaticMeshEditor/Private/SStaticMeshEditorViewport.cpp
--- a/Engine/Source/Editor/StaticMeshEditor/Private/SStaticMeshEditorViewport.cpp
+++ b/Engine/Source/Editor/StaticMeshEditor/Private/SStaticMeshEditorViewport.cpp
@@ -349,7 +349,9 @@ void SStaticMeshEditorViewport::SetViewModeVertexColor()
 	if (!EditorViewportClient->EngineShowFlags.VertexColors)
 	{
 		EditorViewportClient->EngineShowFlags.SetVertexColors(true);
+		EditorViewportClient->EngineShowFlags.SetLighting(false);
 		EditorViewportClient->EngineShowFlags.SetIndirectLightingCache(false);
+		EditorViewportClient->EngineShowFlags.SetPostProcessing(false);
 		EditorViewportClient->EngineShowFlags.SetSelectionOutline(false);
 
 		EditorViewportClient->SetFloorAndEnvironmentVisibility(false);
@@ -361,6 +363,7 @@ void SStaticMeshEditorViewport::SetViewModeVertexColor()
 		EditorViewportClient->EngineShowFlags.SetVertexColors(false);
 		EditorViewportClient->EngineShowFlags.SetLighting(true);
 		EditorViewportClient->EngineShowFlags.SetIndirectLightingCache(true);
+		EditorViewportClient->EngineShowFlags.SetPostProcessing(true);
 		EditorViewportClient->EngineShowFlags.SetSelectionOutline(GetDefault<ULevelEditorViewportSettings>()->bUseSelectionOutline);
 		EditorViewportClient->SetFloorAndEnvironmentVisibility(true);
 		StaticMeshEditorPtr.Pin()->GetStaticMeshComponent()->bDisplayVertexColors = false;
