Fixes Bug UE-79605

StaticMeshEditor now respects the selected LOD even if it is below MinLOD.

If this condition exists a warning will be displayed.

diff --git a/Engine/Source/Editor/StaticMeshEditor/Private/SStaticMeshEditorViewport.cpp b/Engine/Source/Editor/StaticMeshEditor/Private/SStaticMeshEditorViewport.cpp
--- a/Engine/Source/Editor/StaticMeshEditor/Private/SStaticMeshEditorViewport.cpp
+++ b/Engine/Source/Editor/StaticMeshEditor/Private/SStaticMeshEditorViewport.cpp
@@ -58,11 +58,11 @@ void SStaticMeshEditorViewport::Construct(const FArguments& InArgs)
 	SEditorViewport::Construct( SEditorViewport::FArguments() );
 
 	PreviewMeshComponent = NewObject<UStaticMeshComponent>(GetTransientPackage(), NAME_None, RF_Transient );
-	if (GEditor->PreviewPlatform.GetEffectivePreviewFeatureLevel() <= ERHIFeatureLevel::ES3_1)
+	ERHIFeatureLevel::Type FeatureLevel = GEditor->PreviewPlatform.GetEffectivePreviewFeatureLevel();
+	if (FeatureLevel <= ERHIFeatureLevel::ES3_1)
 	{
 		PreviewMeshComponent->SetMobility(EComponentMobility::Static);
 	}
-
 	SetPreviewMesh(StaticMesh);
 
 	FCoreUObjectDelegates::OnObjectPropertyChanged.AddRaw(this, &SStaticMeshEditorViewport::OnObjectPropertyChanged);
@@ -278,11 +278,11 @@ void SStaticMeshEditorViewport::UpdatePreviewMesh(UStaticMesh* InStaticMesh, boo
 	}
 
 	PreviewMeshComponent = NewObject<UStaticMeshComponent>();
-	if (GEditor->PreviewPlatform.GetEffectivePreviewFeatureLevel() <= ERHIFeatureLevel::ES3_1)
+	ERHIFeatureLevel::Type FeatureLevel = GEditor->PreviewPlatform.GetEffectivePreviewFeatureLevel();
+	if ( FeatureLevel <= ERHIFeatureLevel::ES3_1)
 	{
 		PreviewMeshComponent->SetMobility(EComponentMobility::Static);
 	}
-
 	PreviewMeshComponent->SetStaticMesh(InStaticMesh);
 
 	PreviewScene->AddComponent(PreviewMeshComponent,FTransform::Identity);
@@ -411,6 +411,7 @@ void SStaticMeshEditorViewport::OnSetLODModel(int32 InLODSelection)
 {
 	if (PreviewMeshComponent)
 	{
+		PreviewMeshComponent->bOverrideMinLOD = (InLODSelection != 0);
 		LODSelection = InLODSelection;
 		PreviewMeshComponent->SetForcedLodModel(LODSelection);
 		//PopulateUVChoices();
diff --git a/Engine/Source/Editor/StaticMeshEditor/Private/StaticMeshEditorViewportClient.cpp b/Engine/Source/Editor/StaticMeshEditor/Private/StaticMeshEditorViewportClient.cpp
--- a/Engine/Source/Editor/StaticMeshEditor/Private/StaticMeshEditorViewportClient.cpp
+++ b/Engine/Source/Editor/StaticMeshEditor/Private/StaticMeshEditorViewportClient.cpp
@@ -772,20 +772,30 @@ void FStaticMeshEditorViewportClient::DrawCanvas( FViewport& InViewport, FSceneV
 
 	TArray<SStaticMeshEditorViewport::FOverlayTextItem> TextItems;
 
-	int32 CurrentLODLevel = StaticMeshEditor->GetCurrentLODLevel();
-	if (CurrentLODLevel == 0)
+	const int32 CurrentLODLevel = [this, &StaticMeshEditor, &View]()
 	{
-		CurrentLODLevel = ComputeStaticMeshLOD(StaticMesh->RenderData.Get(), StaticMeshComponent->Bounds.Origin, StaticMeshComponent->Bounds.SphereRadius, View, StaticMesh->MinLOD.Default);
-	}
-	else
+		int32 LOD = StaticMeshEditor->GetCurrentLODLevel();
+		return (LOD == 0) ?
+			ComputeStaticMeshLOD(StaticMesh->RenderData.Get(), StaticMeshComponent->Bounds.Origin, StaticMeshComponent->Bounds.SphereRadius, View, StaticMesh->MinLOD.Default)
+			:
+			LOD - 1;
+	}();
+
+	const ERHIFeatureLevel::Type FeatureLevel = GEditor->PreviewPlatform.GetEffectivePreviewFeatureLevel();
+	const int32 CurrentMinLODLevel = StaticMesh->MinLOD.GetValueForFeatureLevel(FeatureLevel);
+	const bool bBelowMinLOD = CurrentLODLevel < CurrentMinLODLevel;
+	TextItems.Add(SStaticMeshEditorViewport::FOverlayTextItem(
+		FText::Format(NSLOCTEXT("UnrealEd", "LOD_F", "LOD:  {0}"), FText::AsNumber(CurrentLODLevel)),
+		bBelowMinLOD ? "TextBlock.ShadowedTextWarning" : "TextBlock.ShadowedText"));
+	
+	if ( bBelowMinLOD )
 	{
-		CurrentLODLevel -= 1;
+		TextItems.Add(SStaticMeshEditorViewport::FOverlayTextItem(
+			FText::Format(NSLOCTEXT("UnrealEd", "BelowMinLODWarning_F", "Selected LOD is below the minimum of {0}"),
+				FText::AsNumber(CurrentMinLODLevel)), "TextBlock.ShadowedTextWarning"));
 	}
 
-	TextItems.Add(SStaticMeshEditorViewport::FOverlayTextItem(
-		FText::Format(NSLOCTEXT("UnrealEd", "LOD_F", "LOD:  {0}"), FText::AsNumber(CurrentLODLevel))));
-
-	float CurrentScreenSize = ComputeBoundsScreenSize(StaticMeshComponent->Bounds.Origin, StaticMeshComponent->Bounds.SphereRadius, View);
+	const float CurrentScreenSize = ComputeBoundsScreenSize(StaticMeshComponent->Bounds.Origin, StaticMeshComponent->Bounds.SphereRadius, View);
 	FNumberFormattingOptions FormatOptions;
 	FormatOptions.MinimumFractionalDigits = 3;
 	FormatOptions.MaximumFractionalDigits = 6;
