Added null check for AssetImportData in FStaticMeshDetails::CustomizeDetails, since it does not need to be set at this point.

diff --git a/Engine/Source/Editor/StaticMeshEditor/Private/StaticMeshEditorTools.cpp b/Engine/Source/Editor/StaticMeshEditor/Private/StaticMeshEditorTools.cpp
--- a/Engine/Source/Editor/StaticMeshEditor/Private/StaticMeshEditorTools.cpp
+++ b/Engine/Source/Editor/StaticMeshEditor/Private/StaticMeshEditorTools.cpp
@@ -93,7 +93,9 @@ void FStaticMeshDetails::CustomizeDetails( class IDetailLayoutBuilder& DetailBui
 	IDetailCategoryBuilder& ImportSettingsCategory = DetailBuilder.EditCategory("ImportSettings");
 
 	TSharedRef<IPropertyHandle> ImportSettings = DetailBuilder.GetProperty(GET_MEMBER_NAME_CHECKED(UStaticMesh, AssetImportData));
-	if (!StaticMeshEditor.GetStaticMesh() || !StaticMeshEditor.GetStaticMesh()->AssetImportData->IsA<UFbxStaticMeshImportData>())
+	if (!StaticMeshEditor.GetStaticMesh() || 
+		!StaticMeshEditor.GetStaticMesh()->AssetImportData ||
+		!StaticMeshEditor.GetStaticMesh()->AssetImportData->IsA<UFbxStaticMeshImportData>())
 	{
 		// Hide the ability to change the import settings object
 		IDetailPropertyRow& Row = ImportSettingsCategory.AddProperty(ImportSettings);
