Moved generation of Layouts menu to SAssetEditoriViewport and moved layouts menu to submenu to match Level Editor

diff --git a/Engine/Source/Editor/EditorStyle/Private/SlateEditorStyle.cpp b/Engine/Source/Editor/EditorStyle/Private/SlateEditorStyle.cpp
--- a/Engine/Source/Editor/EditorStyle/Private/SlateEditorStyle.cpp
+++ b/Engine/Source/Editor/EditorStyle/Private/SlateEditorStyle.cpp
@@ -7269,6 +7269,26 @@ void FSlateEditorStyle::FStyle::SetupToolkitStyles()
 		Set( "StaticMeshEditor.GroupSection", new BOX_BRUSH( "Common/RoundedSelection_16x", FMargin( 4.0f / 16.0f ) ) );
 		Set( "StaticMeshEditor.SetShowVertices", new IMAGE_BRUSH("Icons/icon_StaticMeshEd_Vertices_40x", Icon40x40 ) );
 		Set( "StaticMeshEditor.SetShowVertices.Small", new IMAGE_BRUSH("Icons/icon_StaticMeshEd_Vertices_40x", Icon20x20 ) );
+		Set( "StaticMeshEditor.ToggleShowPivots", new IMAGE_BRUSH("Icons/icon_StaticMeshEd_ShowPivot_40x", Icon40x40));
+		Set( "StaticMeshEditor.ToggleShowPivot.Small", new IMAGE_BRUSH("Icons/icon_StaticMeshEd_ShowPivot_40x", Icon20x20));
+		Set( "StaticMeshEditor.ToggleShowSockets", new IMAGE_BRUSH("Icons/icon_StaticMeshEd_ShowSockets_40x", Icon40x40));
+		Set( "StaticMeshEditor.ToggleShowSockets.Small", new IMAGE_BRUSH("Icons/icon_StaticMeshEd_ShowSockets_40x", Icon20x20));
+		Set( "StaticMeshEditor.ToggleShowNormals", new IMAGE_BRUSH("Icons/icon_StaticMeshEd_Normals_40x", Icon40x40));
+		Set( "StaticMeshEditor.ToggleShowNormals.Small", new IMAGE_BRUSH("Icons/icon_StaticMeshEd_Normals_40x", Icon20x20));
+		Set( "StaticMeshEditor.ToggleShowTangents", new IMAGE_BRUSH("Icons/icon_StaticMeshEd_Tangents_40x", Icon40x40));
+		Set( "StaticMeshEditor.TogglShowTangents.Small", new IMAGE_BRUSH("Icons/icon_StaticMeshEd_Tangents_40x", Icon20x20));
+		Set( "StaticMeshEditor.ToggleShowBinormals", new IMAGE_BRUSH("Icons/icon_StaticMeshEd_Binormals_40x", Icon40x40));
+		Set( "StaticMeshEditor.ToggleShowBinormals.Small", new IMAGE_BRUSH("Icons/icon_StaticMeshEd_Binormals_40x", Icon20x20));
+		Set( "StaticMeshEditor.ToggleShowBounds", new IMAGE_BRUSH("Icons/icon_StaticMeshEd_Bounds_40x", Icon40x40));
+		Set( "StaticMeshEditor.ToggleShowBounds.Small", new IMAGE_BRUSH("Icons/icon_StaticMeshEd_Bounds_40x", Icon20x20));
+		Set( "StaticMeshEditor.ToggleShowGrids", new IMAGE_BRUSH("Icons/icon_MatEd_Grid_40x", Icon40x40));
+		Set( "StaticMeshEditor.ToggleShowGrids.Small", new IMAGE_BRUSH("Icons/icon_MatEd_Grid_40x", Icon20x20));
+		Set( "StaticMeshEditor.ToggleShowVertices", new IMAGE_BRUSH("Icons/icon_StaticMeshEd_Vertices_40x", Icon40x40));
+		Set( "StaticMeshEditor.ToggleShowVertices.Small", new IMAGE_BRUSH("Icons/icon_StaticMeshEd_Vertices_40x", Icon20x20));
+		Set( "StaticMeshEditor.ToggleShowWireframes", new IMAGE_BRUSH("Icons/icon_StaticMeshEd_Wireframe_40x", Icon40x40));
+		Set( "StaticMeshEditor.ToggleShowWireframes.Small", new IMAGE_BRUSH("Icons/icon_StaticMeshEd_Wireframe_40x", Icon20x20));
+		Set( "StaticMeshEditor.ToggleShowVertexColors", new IMAGE_BRUSH("Icons/icon_StaticMeshEd_VertColor_40x", Icon40x40));
+		Set( "StaticMeshEditor.ToggleShowVertexColors.Small", new IMAGE_BRUSH("Icons/icon_StaticMeshEd_VertColor_40x", Icon20x20));
 	}
 
 	// Skeletal Mesh Editor
diff --git a/Engine/Source/Editor/StaticMeshEditor/Private/SStaticMeshEditorViewport.cpp b/Engine/Source/Editor/StaticMeshEditor/Private/SStaticMeshEditorViewport.cpp
--- a/Engine/Source/Editor/StaticMeshEditor/Private/SStaticMeshEditorViewport.cpp
+++ b/Engine/Source/Editor/StaticMeshEditor/Private/SStaticMeshEditorViewport.cpp
@@ -468,7 +468,7 @@ EVisibility SStaticMeshEditorViewport::OnGetViewportContentVisibility() const
 
 void SStaticMeshEditorViewport::BindCommands()
 {
-	SEditorViewport::BindCommands();
+	SAssetEditorViewport::BindCommands();
 
 	const FStaticMeshEditorCommands& Commands = FStaticMeshEditorCommands::Get();
 
diff --git a/Engine/Source/Editor/StaticMeshEditor/Private/SStaticMeshEditorViewport.h b/Engine/Source/Editor/StaticMeshEditor/Private/SStaticMeshEditorViewport.h
--- a/Engine/Source/Editor/StaticMeshEditor/Private/SStaticMeshEditorViewport.h
+++ b/Engine/Source/Editor/StaticMeshEditor/Private/SStaticMeshEditorViewport.h
@@ -12,6 +12,7 @@
 #include "AdvancedPreviewScene.h"
 #include "SEditorViewport.h"
 #include "SCommonEditorViewportToolbarBase.h"
+#include "SAssetEditorViewport.h"
 
 class IStaticMeshEditor;
 class SVerticalBox;
@@ -21,7 +22,7 @@ class UStaticMeshComponent;
 /**
  * StaticMesh Editor Preview viewport widget
  */
-class SStaticMeshEditorViewport : public SEditorViewport, public FGCObject, public ICommonEditorViewportToolbarInfoProvider
+class SStaticMeshEditorViewport : public SAssetEditorViewport, public FGCObject, public ICommonEditorViewportToolbarInfoProvider
 {
 public:
 	SLATE_BEGIN_ARGS( SStaticMeshEditorViewport ){}
@@ -75,7 +76,7 @@ public:
 	/** Retrieves the static mesh component. */
 	UStaticMeshComponent* GetStaticMeshComponent() const;
 
-	/** 
+	/**
 	 *	Sets up the static mesh that the Static Mesh editor is viewing.
 	 *
 	 *	@param	InStaticMesh		The static mesh being viewed in the editor.
@@ -129,25 +130,28 @@ protected:
 	virtual TSharedPtr<SWidget> MakeViewportToolbar() override;
 	virtual void PopulateViewportOverlays(TSharedRef<SOverlay> Overlay) override;
 
-private:
-	/** Determines the visibility of the viewport. */
-	bool IsVisible() const override;
+public:
+
+	/** Callback for toggling the vertex color show flag. */
+	void SetViewModeVertexColor();
+
+	/** Callback for checking the vertex color show flag. */
+	bool IsInViewModeVertexColorChecked() const;
 
-	/** Callback for toggling the wireframe mode flag. */
+	/** Callback for toggling the wireframe mode. */
 	void SetViewModeWireframe();
-	
-	/** Callback for checking the wireframe mode flag. */
+
+	/** Callback for checking the wireframe flag. */
 	bool IsInViewModeWireframeChecked() const;
 
-	/** Callback for toggling the vertex color show flag. */
-	void SetViewModeVertexColor();
+private:
+
+	/** Determines the visibility of the viewport. */
+	bool IsVisible() const override;
 
 	/** Implementation of the SetViewModeVertexColor, used to apply the vertex color show flag. */
 	void SetViewModeVertexColorImplemetation(bool bValue);
 
-	/** Callback for checking the vertex color show flag. */
-	bool IsInViewModeVertexColorChecked() const;
-
 	/** Callback for toggling the realtime preview flag. */
 	void SetRealtimePreview();
 
diff --git a/Engine/Source/Editor/StaticMeshEditor/Private/StaticMeshEditor.cpp b/Engine/Source/Editor/StaticMeshEditor/Private/StaticMeshEditor.cpp
--- a/Engine/Source/Editor/StaticMeshEditor/Private/StaticMeshEditor.cpp
+++ b/Engine/Source/Editor/StaticMeshEditor/Private/StaticMeshEditor.cpp
@@ -45,6 +45,8 @@
 #include "ConvexDecompositionNotification.h"
 #include "FbxMeshUtils.h"
 #include "RawMesh.h"
+#include "EditorViewportTabContent.h"
+#include "EditorViewportLayout.h"
 
 #define LOCTEXT_NAMESPACE "StaticMeshEditor"
 
@@ -165,9 +167,24 @@ void FStaticMeshEditor::InitEditorForStaticMesh(UStaticMesh* ObjectToEdit)
 
 	BindCommands();
 
-	Viewport = SNew(SStaticMeshEditorViewport)
-		.StaticMeshEditor(SharedThis(this))
-		.ObjectToEdit(ObjectToEdit);
+  	TWeakPtr<IStaticMeshEditor> WeakSharedThis(SharedThis(this));
+	MakeViewportFunc = [=]()
+	{
+		return SNew(SStaticMeshEditorViewport)
+			.StaticMeshEditor(WeakSharedThis)
+			.ObjectToEdit(ObjectToEdit);
+	};
+
+	// The tab must be created before the viewport layout because the layout needs them
+	TSharedRef< SDockTab > DockableTab =
+		SNew(SDockTab)
+		.Icon(FEditorStyle::GetBrush("LevelEditor.Tabs.Viewports"));
+
+	// Create a new tab
+	ViewportTabContent = MakeShareable(new FEditorViewportTabContent());
+
+	const FString LayoutId = FString("StaticMeshEditorViewport");
+	ViewportTabContent->Initialize(MakeViewportFunc, DockableTab, LayoutId);
 
 	FPropertyEditorModule& PropertyEditorModule = FModuleManager::LoadModuleChecked<FPropertyEditorModule>(TEXT("PropertyEditor"));
 
@@ -183,7 +200,7 @@ void FStaticMeshEditor::InitEditorForStaticMesh(UStaticMesh* ObjectToEdit)
 	FOnGetDetailCustomizationInstance LayoutCustomStaticMeshProperties = FOnGetDetailCustomizationInstance::CreateSP( this, &FStaticMeshEditor::MakeStaticMeshDetails );
 	StaticMeshDetailsView->RegisterInstancedCustomPropertyLayout( UStaticMesh::StaticClass(), LayoutCustomStaticMeshProperties );
 
-	SetEditorMesh(ObjectToEdit);
+ 	SetEditorMesh(ObjectToEdit);
 }
 
 void FStaticMeshEditor::InitStaticMeshEditor( const EToolkitMode::Type Mode, const TSharedPtr< class IToolkitHost >& InitToolkitHost, UStaticMesh* ObjectToEdit )
@@ -401,6 +418,7 @@ void FStaticMeshEditor::ExtendMenu()
 				InMenuBuilder.AddMenuEntry(FStaticMeshEditorCommands::Get().CreateDOP26);
 				InMenuBuilder.AddMenuEntry(FStaticMeshEditorCommands::Get().ConvertBoxesToConvex);
 				InMenuBuilder.AddMenuEntry(FStaticMeshEditorCommands::Get().RemoveCollision);
+				
 				InMenuBuilder.AddMenuEntry(FGenericCommands::Get().Delete, "DeleteCollision", LOCTEXT("DeleteCollision", "Delete Selected Collision"), LOCTEXT("DeleteCollisionToolTip", "Deletes the selected Collision from the mesh."));
 				InMenuBuilder.AddMenuEntry(FGenericCommands::Get().Duplicate, "DuplicateCollision", LOCTEXT("DuplicateCollision", "Duplicate Selected Collision"), LOCTEXT("DuplicateCollisionToolTip", "Duplicates the selected Collision."));
 			}
@@ -463,18 +481,14 @@ void FStaticMeshEditor::AddReferencedObjects( FReferenceCollector& Collector )
 
 TSharedRef<SDockTab> FStaticMeshEditor::SpawnTab_Viewport( const FSpawnTabArgs& Args )
 {
-	check( Args.GetTabId() == ViewportTabId );
-
-	TSharedRef<SDockTab> SpawnedTab =
-	 SNew(SDockTab)
-		.Label( LOCTEXT("StaticMeshViewport_TabTitle", "Viewport") )
-		[
-			Viewport.ToSharedRef()
-		];
+	TSharedRef< SDockTab > DockableTab =
+		SNew(SDockTab)
+		.Icon(FEditorStyle::GetBrush("LevelEditor.Tabs.Viewports"));
 
-	Viewport->SetParentTab( SpawnedTab );
+	const FString LayoutId = FString("StaticMeshEditorViewport");
+ 	ViewportTabContent->Initialize(MakeViewportFunc, DockableTab, LayoutId);
 
-	return SpawnedTab;
+	return DockableTab;
 }
 
 TSharedRef<SDockTab> FStaticMeshEditor::SpawnTab_Properties( const FSpawnTabArgs& Args )
@@ -545,12 +559,19 @@ TSharedRef<SDockTab> FStaticMeshEditor::SpawnTab_SecondaryToolbar( const FSpawnT
 	return SpawnedTab;
 }
 
+TSharedPtr<class SStaticMeshEditorViewport> FStaticMeshEditor::GetStaticMeshViewport() const
+{
+	// we can use static cast here b/c we know in this editor we will have a static mesh viewport 
+	return StaticCastSharedPtr<SStaticMeshEditorViewport>(ViewportTabContent->GetFirstViewport());
+}
+
 void FStaticMeshEditor::BindCommands()
 {
 	const FStaticMeshEditorCommands& Commands = FStaticMeshEditorCommands::Get();
 
 	const TSharedRef<FUICommandList>& UICommandList = GetToolkitCommands();
 
+
 	UICommandList->MapAction( FGenericCommands::Get().Delete,
 		FExecuteAction::CreateSP( this, &FStaticMeshEditor::DeleteSelected ),
 		FCanExecuteAction::CreateSP(this, &FStaticMeshEditor::CanDeleteSelected));
@@ -561,6 +582,7 @@ void FStaticMeshEditor::BindCommands()
 	UICommandList->MapAction( FGenericCommands::Get().Redo,
 		FExecuteAction::CreateSP( this, &FStaticMeshEditor::RedoAction ) );
 
+
 	UICommandList->MapAction(
 		FGenericCommands::Get().Duplicate,
 		FExecuteAction::CreateSP(this, &FStaticMeshEditor::DuplicateSelected),
@@ -603,6 +625,66 @@ void FStaticMeshEditor::BindCommands()
 		Commands.CreateSphylCollision,
 		FExecuteAction::CreateSP(this, &FStaticMeshEditor::OnCollisionSphyl));
 
+	UICommandList->MapAction(Commands.ToggleShowNormals,
+		FExecuteAction::CreateSP(this, &FStaticMeshEditor::ToggleShowNormals),
+		FCanExecuteAction(),
+		FIsActionChecked::CreateSP(this, &FStaticMeshEditor::IsShowNormalsChecked));
+
+	UICommandList->MapAction(Commands.ToggleShowTangents,
+		FExecuteAction::CreateSP(this, &FStaticMeshEditor::ToggleShowTangents),
+		FCanExecuteAction(),
+		FIsActionChecked::CreateSP(this, &FStaticMeshEditor::IsShowTangentsChecked));
+
+	UICommandList->MapAction(Commands.ToggleShowBinormals,
+		FExecuteAction::CreateSP(this, &FStaticMeshEditor::ToggleShowBinormals),
+		FCanExecuteAction(),
+		FIsActionChecked::CreateSP(this, &FStaticMeshEditor::IsShowBinormalsChecked));
+
+	UICommandList->MapAction(Commands.ToggleShowPivots,
+		FExecuteAction::CreateSP(this, &FStaticMeshEditor::ToggleShowPivots),
+		FCanExecuteAction(),
+		FIsActionChecked::CreateSP(this, &FStaticMeshEditor::IsShowPivotsChecked));
+
+	UICommandList->MapAction(Commands.ToggleShowVertices,
+		FExecuteAction::CreateSP(this, &FStaticMeshEditor::ToggleShowVertices),
+		FCanExecuteAction(),
+		FIsActionChecked::CreateSP(this, &FStaticMeshEditor::IsShowVerticesChecked));
+
+	UICommandList->MapAction(Commands.ToggleShowGrids,
+		FExecuteAction::CreateSP(this, &FStaticMeshEditor::ToggleShowGrids),
+		FCanExecuteAction(),
+		FIsActionChecked::CreateSP(this, &FStaticMeshEditor::IsShowGridsChecked));
+
+	UICommandList->MapAction(Commands.ToggleShowBounds,
+		FExecuteAction::CreateSP(this, &FStaticMeshEditor::ToggleShowBounds),
+		FCanExecuteAction(),
+		FIsActionChecked::CreateSP(this, &FStaticMeshEditor::IsShowBoundsChecked));
+
+	UICommandList->MapAction(Commands.ToggleShowSimpleCollisions,
+		FExecuteAction::CreateSP(this, &FStaticMeshEditor::ToggleShowSimpleCollisions),
+		FCanExecuteAction(),
+		FIsActionChecked::CreateSP(this, &FStaticMeshEditor::IsShowSimpleCollisionsChecked));
+
+	UICommandList->MapAction(Commands.ToggleShowComplexCollisions,
+		FExecuteAction::CreateSP(this, &FStaticMeshEditor::ToggleShowComplexCollisions),
+		FCanExecuteAction(),
+		FIsActionChecked::CreateSP(this, &FStaticMeshEditor::IsShowComplexCollisionsChecked));
+
+	UICommandList->MapAction(Commands.ToggleShowSockets,
+		FExecuteAction::CreateSP(this, &FStaticMeshEditor::ToggleShowSockets),
+		FCanExecuteAction(),
+		FIsActionChecked::CreateSP(this, &FStaticMeshEditor::IsShowSocketsChecked));
+
+	UICommandList->MapAction(Commands.ToggleShowWireframes,
+		FExecuteAction::CreateSP(this, &FStaticMeshEditor::ToggleShowWireframes),
+		FCanExecuteAction(),
+		FIsActionChecked::CreateSP(this, &FStaticMeshEditor::IsShowWireframesChecked));
+
+	UICommandList->MapAction(Commands.ToggleShowVertexColors,
+		FExecuteAction::CreateSP(this, &FStaticMeshEditor::ToggleShowVertexColors),
+		FCanExecuteAction(),
+		FIsActionChecked::CreateSP(this, &FStaticMeshEditor::IsShowVertexColorsChecked));
+
 	UICommandList->MapAction(
 		Commands.RemoveCollision,
 		FExecuteAction::CreateSP(this, &FStaticMeshEditor::OnRemoveCollision),
@@ -636,6 +718,8 @@ void FStaticMeshEditor::BindCommands()
 	UICommandList->MapAction(
 		Commands.CreateAutoConvexCollision,
 		FExecuteAction::CreateSP(this, &FStaticMeshEditor::OnConvexDecomposition));
+
+
 }
 
 static TSharedRef< SWidget > GenerateCollisionMenuContent(TSharedPtr<const FUICommandList> InCommandList)
@@ -644,8 +728,8 @@ static TSharedRef< SWidget > GenerateCollisionMenuContent(TSharedPtr<const FUICo
 
 	MenuBuilder.BeginSection("ShowCollision", LOCTEXT("ShowCollision", "Show Collision"));
 	{
-		MenuBuilder.AddMenuEntry(FStaticMeshEditorCommands::Get().SetShowSimpleCollision);
-		MenuBuilder.AddMenuEntry(FStaticMeshEditorCommands::Get().SetShowComplexCollision);
+		MenuBuilder.AddMenuEntry(FStaticMeshEditorCommands::Get().ToggleShowSimpleCollisions);
+		MenuBuilder.AddMenuEntry(FStaticMeshEditorCommands::Get().ToggleShowComplexCollisions);
 	}
 	MenuBuilder.EndSection();
 
@@ -677,7 +761,7 @@ void FStaticMeshEditor::ExtendToolBar()
 				ToolbarBuilder.AddToolBarButton(FEditorViewportCommands::Get().ToggleRealTime);
 			}
 			ToolbarBuilder.EndSection();
-			
+
 			ToolbarBuilder.BeginSection("Mesh");
 			{
 				ToolbarBuilder.AddToolBarButton(FUIAction(FExecuteAction::CreateSP(ThisEditor, &FStaticMeshEditor::HandleReimportMesh)),
@@ -693,14 +777,14 @@ void FStaticMeshEditor::ExtendToolBar()
 				);
 			}
 			ToolbarBuilder.EndSection();
-	
+
 			ToolbarBuilder.BeginSection("Command");
 			{
-				ToolbarBuilder.AddToolBarButton(FStaticMeshEditorCommands::Get().SetShowSockets);
-				ToolbarBuilder.AddToolBarButton(FStaticMeshEditorCommands::Get().SetShowWireframe);
-				ToolbarBuilder.AddToolBarButton(FStaticMeshEditorCommands::Get().SetShowVertexColor);
-				ToolbarBuilder.AddToolBarButton(FStaticMeshEditorCommands::Get().SetShowGrid);
-				ToolbarBuilder.AddToolBarButton(FStaticMeshEditorCommands::Get().SetShowBounds);
+				ToolbarBuilder.AddToolBarButton(FStaticMeshEditorCommands::Get().ToggleShowSockets);
+				ToolbarBuilder.AddToolBarButton(FStaticMeshEditorCommands::Get().ToggleShowWireframes);
+				ToolbarBuilder.AddToolBarButton(FStaticMeshEditorCommands::Get().ToggleShowVertexColors);
+				ToolbarBuilder.AddToolBarButton(FStaticMeshEditorCommands::Get().ToggleShowGrids);
+				ToolbarBuilder.AddToolBarButton(FStaticMeshEditorCommands::Get().ToggleShowBounds);
 
 				TSharedPtr<const FUICommandList> CommandList = ToolbarBuilder.GetTopCommandList();
 
@@ -712,11 +796,11 @@ void FStaticMeshEditor::ExtendToolBar()
 					FSlateIcon(FEditorStyle::GetStyleSetName(), "StaticMeshEditor.SetShowCollision")
 				);
 
-				ToolbarBuilder.AddToolBarButton(FStaticMeshEditorCommands::Get().SetShowPivot);
-				ToolbarBuilder.AddToolBarButton(FStaticMeshEditorCommands::Get().SetShowNormals);
-				ToolbarBuilder.AddToolBarButton(FStaticMeshEditorCommands::Get().SetShowTangents);
-				ToolbarBuilder.AddToolBarButton(FStaticMeshEditorCommands::Get().SetShowBinormals);
-				ToolbarBuilder.AddToolBarButton(FStaticMeshEditorCommands::Get().SetShowVertices);
+				ToolbarBuilder.AddToolBarButton(FStaticMeshEditorCommands::Get().ToggleShowPivots);
+				ToolbarBuilder.AddToolBarButton(FStaticMeshEditorCommands::Get().ToggleShowNormals);
+				ToolbarBuilder.AddToolBarButton(FStaticMeshEditorCommands::Get().ToggleShowTangents);
+				ToolbarBuilder.AddToolBarButton(FStaticMeshEditorCommands::Get().ToggleShowBinormals);
+				ToolbarBuilder.AddToolBarButton(FStaticMeshEditorCommands::Get().ToggleShowVertices);
 
 				FOnGetContent OnGetUVMenuContent = FOnGetContent::CreateRaw(ThisEditor, &FStaticMeshEditor::GenerateUVChannelComboList);
 
@@ -729,7 +813,6 @@ void FStaticMeshEditor::ExtendToolBar()
 			}
 
 			ToolbarBuilder.EndSection();
-
 			ToolbarBuilder.BeginSection("Camera");
 			{
 				ToolbarBuilder.AddToolBarButton(FStaticMeshEditorCommands::Get().ResetCamera);
@@ -740,25 +823,29 @@ void FStaticMeshEditor::ExtendToolBar()
 		}
 	};
 
+
+	{
 	TSharedPtr<FExtender> ToolbarExtender = MakeShareable(new FExtender);
 
-	FStaticMeshEditorViewportClient& ViewportClient = Viewport->GetViewportClient();
+	FStaticMeshEditorViewportClient& ViewportClient = GetStaticMeshViewport()->GetViewportClient();
 
 	FStaticMeshEditor* ThisEditor = this;
 
 	ToolbarExtender->AddToolBarExtension(
 		"Asset",
 		EExtensionHook::After,
-		Viewport->GetCommandList(),
+		GetToolkitCommands(),
 		FToolBarExtensionDelegate::CreateStatic(&Local::FillToolbar, ThisEditor)
-		);
+	);
 
 	AddToolbarExtender(ToolbarExtender);
 
-	IStaticMeshEditorModule* StaticMeshEditorModule = &FModuleManager::LoadModuleChecked<IStaticMeshEditorModule>( "StaticMeshEditor" );
+	IStaticMeshEditorModule* StaticMeshEditorModule = &FModuleManager::LoadModuleChecked<IStaticMeshEditorModule>("StaticMeshEditor");
 	EditorToolbarExtender = StaticMeshEditorModule->GetToolBarExtensibilityManager()->GetAllExtenders(GetToolkitCommands(), GetEditingObjects());
 	AddToolbarExtender(EditorToolbarExtender);
 	AddSecondaryToolbarExtender(StaticMeshEditorModule->GetSecondaryToolBarExtensibilityManager()->GetAllExtenders(GetToolkitCommands(), GetEditingObjects()));
+	}
+
 }
 
 void FStaticMeshEditor::BuildSubTools()
@@ -771,7 +858,7 @@ void FStaticMeshEditor::BuildSubTools()
 		.StaticMeshEditorPtr(SharedThis(this));
 
 	FAdvancedPreviewSceneModule& AdvancedPreviewSceneModule = FModuleManager::LoadModuleChecked<FAdvancedPreviewSceneModule>("AdvancedPreviewScene");
-	AdvancedPreviewSettingsWidget = AdvancedPreviewSceneModule.CreateAdvancedPreviewSceneSettingsWidget(Viewport->GetPreviewScene());
+	AdvancedPreviewSettingsWidget = AdvancedPreviewSceneModule.CreateAdvancedPreviewSceneSettingsWidget(GetStaticMeshViewport()->GetPreviewScene());
 }
 
 FName FStaticMeshEditor::GetToolkitFName() const
@@ -796,7 +883,7 @@ FLinearColor FStaticMeshEditor::GetWorldCentricTabColorScale() const
 
 UStaticMeshComponent* FStaticMeshEditor::GetStaticMeshComponent() const
 {
-	return Viewport->GetStaticMeshComponent();
+	return GetStaticMeshViewport()->GetStaticMeshComponent();
 }
 
 void FStaticMeshEditor::SetSelectedSocket(UStaticMeshSocket* InSelectedSocket)
@@ -850,9 +937,9 @@ void FStaticMeshEditor::AddSelectedPrim(const FPrimData& InPrimData, bool bClear
 	check(IsPrimValid(InPrimData));
 
 	// Enable collision, if not already
-	if( !Viewport->GetViewportClient().IsShowSimpleCollisionChecked() )
+	if( !GetStaticMeshViewport()->GetViewportClient().IsShowSimpleCollisionChecked() )
 	{
-		Viewport->GetViewportClient().ToggleShowSimpleCollision();
+		GetStaticMeshViewport()->GetViewportClient().ToggleShowSimpleCollision();
 	}
 
 	if( bClearSelection )
@@ -957,7 +1044,7 @@ void FStaticMeshEditor::DuplicateSelectedPrims(const FVector* InOffset)
 		StaticMesh->MarkPackageDirty();
 
 		// Update views/property windows
-		Viewport->RefreshViewport();
+		GetStaticMeshViewport()->RefreshViewport();
 
 		StaticMesh->bCustomizedCollision = true;	//mark the static mesh for collision customization
 	}
@@ -1276,7 +1363,7 @@ void FStaticMeshEditor::RefreshTool()
 
 void FStaticMeshEditor::RefreshViewport()
 {
-	Viewport->RefreshViewport();
+	GetStaticMeshViewport()->RefreshViewport();
 }
 
 TSharedRef<SWidget> FStaticMeshEditor::GenerateUVChannelComboList()
@@ -1284,7 +1371,7 @@ TSharedRef<SWidget> FStaticMeshEditor::GenerateUVChannelComboList()
 	FMenuBuilder MenuBuilder(true, nullptr, EditorToolbarExtender);
 	FUIAction DrawUVsAction;
 
-	FStaticMeshEditorViewportClient& ViewportClient = Viewport->GetViewportClient();
+	FStaticMeshEditorViewportClient& ViewportClient = GetStaticMeshViewport()->GetViewportClient();
 
 	DrawUVsAction.ExecuteAction = FExecuteAction::CreateRaw(&ViewportClient, &FStaticMeshEditorViewportClient::SetDrawUVOverlay, false);
 
@@ -1368,7 +1455,7 @@ void FStaticMeshEditor::UpdateLODStats(int32 CurrentLOD)
 
 void FStaticMeshEditor::ComboBoxSelectionChanged( TSharedPtr<FString> NewSelection, ESelectInfo::Type /*SelectInfo*/ )
 {
-	Viewport->RefreshViewport();
+	GetStaticMeshViewport()->RefreshViewport();
 }
 
 void FStaticMeshEditor::HandleReimportMesh()
@@ -1466,7 +1553,7 @@ void FStaticMeshEditor::GenerateKDop(const FVector* Directions, uint32 NumDirect
 		// Don't 'nudge' KDop prims, as they are fitted specifically around the geometry
 	}
 
-	Viewport->RefreshViewport();
+	GetStaticMeshViewport()->RefreshViewport();
 }
 
 void FStaticMeshEditor::OnCollisionBox()
@@ -1489,7 +1576,7 @@ void FStaticMeshEditor::OnCollisionBox()
 		}
 	}
 
-	Viewport->RefreshViewport();
+	GetStaticMeshViewport()->RefreshViewport();
 }
 
 void FStaticMeshEditor::OnCollisionSphere()
@@ -1512,7 +1599,7 @@ void FStaticMeshEditor::OnCollisionSphere()
 		}
 	}
 
-	Viewport->RefreshViewport();
+	GetStaticMeshViewport()->RefreshViewport();
 }
 
 void FStaticMeshEditor::OnCollisionSphyl()
@@ -1535,7 +1622,7 @@ void FStaticMeshEditor::OnCollisionSphyl()
 		}
 	}
 
-	Viewport->RefreshViewport();
+	GetStaticMeshViewport()->RefreshViewport();
 }
 
 void FStaticMeshEditor::OnRemoveCollision(void)
@@ -1562,7 +1649,7 @@ void FStaticMeshEditor::OnRemoveCollision(void)
 	StaticMesh->MarkPackageDirty();
 
 	// Update views/property windows
-	Viewport->RefreshViewport();
+	GetStaticMeshViewport()->RefreshViewport();
 
 	StaticMesh->bCustomizedCollision = true;	//mark the static mesh for collision customization
 }
@@ -1650,7 +1737,7 @@ void FStaticMeshEditor::OnConvertBoxToConvexCollision()
 				StaticMesh->MarkPackageDirty();
 
 				// Update views/property windows
-				Viewport->RefreshViewport();
+				GetStaticMeshViewport()->RefreshViewport();
 
 				StaticMesh->bCustomizedCollision = true;	//mark the static mesh for collision customization
 			}
@@ -1677,9 +1764,9 @@ void FStaticMeshEditor::OnCopyCollisionFromSelectedStaticMesh()
 	BodySetup->CopyBodyPropertiesFrom(SelectedMesh->BodySetup);
 
 	// Enable collision, if not already
-	if( !Viewport->GetViewportClient().IsShowSimpleCollisionChecked() )
+	if( !GetStaticMeshViewport()->GetViewportClient().IsShowSimpleCollisionChecked() )
 	{
-		Viewport->GetViewportClient().ToggleShowSimpleCollision();
+		GetStaticMeshViewport()->GetViewportClient().ToggleShowSimpleCollision();
 	}
 
 	// Invalidate physics data and create new meshes
@@ -1697,7 +1784,7 @@ void FStaticMeshEditor::OnCopyCollisionFromSelectedStaticMesh()
 	GUnrealEd->RedrawLevelEditingViewports();
 
 	// Update views/property windows
-	Viewport->RefreshViewport();
+	GetStaticMeshViewport()->RefreshViewport();
 
 	StaticMesh->bCustomizedCollision = true;	//mark the static mesh for collision customization
 }
@@ -1768,8 +1855,8 @@ void FStaticMeshEditor::SetEditorMesh(UStaticMesh* InStaticMesh, bool bResetCame
 	// Set the details view.
 	StaticMeshDetailsView->SetObject(StaticMesh);
 
-	Viewport->UpdatePreviewMesh(StaticMesh, bResetCamera);
-	Viewport->RefreshViewport();
+ 	GetStaticMeshViewport()->UpdatePreviewMesh(StaticMesh, bResetCamera);
+ 	GetStaticMeshViewport()->RefreshViewport();
 }
 
 void FStaticMeshEditor::OnChangeMesh()
@@ -1825,7 +1912,7 @@ void FStaticMeshEditor::OnSaveGeneratedLODs()
 		StaticMeshDetailsView->SetObjects(Selected, true);
 
 		// Update screen
-		Viewport->RefreshViewport();
+		GetStaticMeshViewport()->RefreshViewport();
 	}
 }
 
@@ -1901,9 +1988,9 @@ void FStaticMeshEditor::DoDecomp(uint32 InHullCount, int32 InMaxHullVerts, uint3
 		}
 
 		// Enable collision, if not already
-		if( !Viewport->GetViewportClient().IsShowSimpleCollisionChecked() )
+		if( !GetStaticMeshViewport()->GetViewportClient().IsShowSimpleCollisionChecked() )
 		{
-			Viewport->GetViewportClient().ToggleShowSimpleCollision();
+			GetStaticMeshViewport()->GetViewportClient().ToggleShowSimpleCollision();
 		}
 
 		// refresh collision change back to staticmesh components
@@ -1913,7 +2000,7 @@ void FStaticMeshEditor::DoDecomp(uint32 InHullCount, int32 InMaxHullVerts, uint3
 		StaticMesh->MarkPackageDirty();
 
 		// Update screen.
-		Viewport->RefreshViewport();
+		GetStaticMeshViewport()->RefreshViewport();
 
 		StaticMesh->bCustomizedCollision = true;	//mark the static mesh for collision customization
 	}
@@ -1921,7 +2008,7 @@ void FStaticMeshEditor::DoDecomp(uint32 InHullCount, int32 InMaxHullVerts, uint3
 
 TSet< int32 >& FStaticMeshEditor::GetSelectedEdges()
 {
-	return Viewport->GetSelectedEdges();
+	return GetStaticMeshViewport()->GetSelectedEdges();
 }
 
 int32 FStaticMeshEditor::GetNumTriangles( int32 LODLevel ) const
@@ -2024,7 +2111,7 @@ void FStaticMeshEditor::DeleteSelectedPrims()
 		StaticMesh->MarkPackageDirty();
 
 		// Update views/property windows
-		Viewport->RefreshViewport();
+		GetStaticMeshViewport()->RefreshViewport();
 
 		StaticMesh->bCustomizedCollision = true;	//mark the static mesh for collision customization
 	}
@@ -2090,9 +2177,9 @@ void FStaticMeshEditor::OnObjectReimported(UObject* InObject)
 
 EViewModeIndex FStaticMeshEditor::GetViewMode() const
 {
-	if (Viewport.IsValid())
+	if (GetStaticMeshViewport().IsValid())
 	{
-		const FStaticMeshEditorViewportClient& ViewportClient = Viewport->GetViewportClient();
+		const FStaticMeshEditorViewportClient& ViewportClient = GetStaticMeshViewport()->GetViewportClient();
 		return ViewportClient.GetViewMode();
 	}
 	else
@@ -2103,7 +2190,7 @@ EViewModeIndex FStaticMeshEditor::GetViewMode() const
 
 FEditorViewportClient& FStaticMeshEditor::GetViewportClient()
 {
-	return Viewport->GetViewportClient();
+	return GetStaticMeshViewport()->GetViewportClient();
 }
 
 void FStaticMeshEditor::OnConvexDecomposition()
@@ -2211,7 +2298,7 @@ void FStaticMeshEditor::OnSocketSelectionChanged()
 	{
 		ClearSelectedPrims();
 	}
-	Viewport->GetViewportClient().OnSocketSelectionChanged( SelectedSocket );
+	GetStaticMeshViewport()->GetViewportClient().OnSocketSelectionChanged( SelectedSocket );
 }
 
 void FStaticMeshEditor::OnPostReimport(UObject* InObject, bool bSuccess)
@@ -2231,12 +2318,12 @@ void FStaticMeshEditor::OnPostReimport(UObject* InObject, bool bSuccess)
 void FStaticMeshEditor::SetCurrentViewedUVChannel(int32 InNewUVChannel)
 {
 	CurrentViewedUVChannel = FMath::Clamp(InNewUVChannel, 0, GetNumUVChannels());
-	Viewport->GetViewportClient().SetDrawUVOverlay(true);
+	GetStaticMeshViewport()->GetViewportClient().SetDrawUVOverlay(true);
 }
 
 ECheckBoxState FStaticMeshEditor::GetUVChannelCheckState(int32 TestUVChannel) const
 {
-	return CurrentViewedUVChannel == TestUVChannel && Viewport->GetViewportClient().IsDrawUVOverlayChecked() ? ECheckBoxState::Checked : ECheckBoxState::Unchecked;
+	return CurrentViewedUVChannel == TestUVChannel && GetStaticMeshViewport()->GetViewportClient().IsDrawUVOverlayChecked() ? ECheckBoxState::Checked : ECheckBoxState::Unchecked;
 }
 
 void FStaticMeshEditor::Tick(float DeltaTime)
@@ -2269,10 +2356,315 @@ bool FStaticMeshEditor::CanRemoveUVChannel()
 {
 	// Can remove UV channel if there's one that is currently being selected and displayed, 
 	// and the current LOD has more than one UV channel
-	return Viewport->GetViewportClient().IsDrawUVOverlayChecked() && 
+	return GetStaticMeshViewport()->GetViewportClient().IsDrawUVOverlayChecked() && 
 		StaticMesh->GetNumUVChannels(GetCurrentLODIndex()) > 1;
 }
 
+void FStaticMeshEditor::ToggleShowNormals()
+{
+	bDrawNormals = !bDrawNormals;
+	TFunction<void(FName, TSharedPtr<IEditorViewportLayoutEntity>)> ToggleShowNormalsFunc = 
+		[this](FName Name, TSharedPtr<IEditorViewportLayoutEntity> Entity)
+	{
+		TSharedRef<SStaticMeshEditorViewport> StaticMeshEditorViewport = StaticCastSharedRef<SStaticMeshEditorViewport>(Entity->AsWidget());
+		FStaticMeshEditorViewportClient& StaticMeshEditorViewportClient  = StaticMeshEditorViewport->GetViewportClient();
+		StaticMeshEditorViewportClient.SetShowNormals(bDrawNormals);
+	};
+
+  	ViewportTabContent->PerformActionOnViewports(ToggleShowNormalsFunc);
+
+	if (FEngineAnalytics::IsAvailable())
+	{
+		FEngineAnalytics::GetProvider().RecordEvent(TEXT("Editor.Usage.StaticMesh.Toolbar"), TEXT("bDrawNormals"), bDrawNormals ? TEXT("True") : TEXT("False"));
+	}
+}
+
+bool FStaticMeshEditor::IsShowNormalsChecked() const
+{
+	return bDrawNormals;
+}
+
+void FStaticMeshEditor::ToggleShowTangents()
+{
+	bDrawTangents = !bDrawTangents;
+	TFunction<void(FName, TSharedPtr<IEditorViewportLayoutEntity>)> ToggleShowTangentsFunc =
+		[this](FName Name, TSharedPtr<IEditorViewportLayoutEntity> Entity)
+	{
+		TSharedRef<SStaticMeshEditorViewport> StaticMeshEditorViewport = StaticCastSharedRef<SStaticMeshEditorViewport>(Entity->AsWidget());
+		FStaticMeshEditorViewportClient& StaticMeshEditorViewportClient = StaticMeshEditorViewport->GetViewportClient();
+		StaticMeshEditorViewportClient.SetShowTangents(bDrawTangents);
+	};
+
+	ViewportTabContent->PerformActionOnViewports(ToggleShowTangentsFunc);
+
+	if (FEngineAnalytics::IsAvailable())
+	{
+		FEngineAnalytics::GetProvider().RecordEvent(TEXT("Editor.Usage.StaticMesh.Toolbar"), TEXT("bDrawTangents"), bDrawTangents ? TEXT("True") : TEXT("False"));
+	}
+}
+
+bool FStaticMeshEditor::IsShowTangentsChecked() const
+{
+	return bDrawTangents;
+}
+
+void FStaticMeshEditor::ToggleShowBinormals()
+{
+	bDrawBinormals= !bDrawBinormals;
+	TFunction<void(FName, TSharedPtr<IEditorViewportLayoutEntity>)> ToggleShowBinormalsFunc =
+		[this](FName Name, TSharedPtr<IEditorViewportLayoutEntity> Entity)
+	{
+		TSharedRef<SStaticMeshEditorViewport> StaticMeshEditorViewport = StaticCastSharedRef<SStaticMeshEditorViewport>(Entity->AsWidget());
+		FStaticMeshEditorViewportClient& StaticMeshEditorViewportClient = StaticMeshEditorViewport->GetViewportClient();
+		StaticMeshEditorViewportClient.SetShowBinormals(bDrawBinormals);
+	};
+
+	ViewportTabContent->PerformActionOnViewports(ToggleShowBinormalsFunc);
+
+	if (FEngineAnalytics::IsAvailable())
+	{
+		FEngineAnalytics::GetProvider().RecordEvent(TEXT("Editor.Usage.StaticMesh.Toolbar"), TEXT("bDrawBinormals"), bDrawBinormals? TEXT("True") : TEXT("False"));
+	}
+}
+
+bool FStaticMeshEditor::IsShowBinormalsChecked() const
+{
+	return bDrawBinormals;
+}
+
+void FStaticMeshEditor::ToggleShowPivots()
+{
+	bDrawPivots = !bDrawPivots;
+	TFunction<void(FName, TSharedPtr<IEditorViewportLayoutEntity>)> ToggleShowPivotsFunc =
+		[this](FName Name, TSharedPtr<IEditorViewportLayoutEntity> Entity)
+	{
+		TSharedRef<SStaticMeshEditorViewport> StaticMeshEditorViewport = StaticCastSharedRef<SStaticMeshEditorViewport>(Entity->AsWidget());
+		FStaticMeshEditorViewportClient& StaticMeshEditorViewportClient = StaticMeshEditorViewport->GetViewportClient();
+		StaticMeshEditorViewportClient.SetShowPivots(bDrawPivots);
+	};
+
+	ViewportTabContent->PerformActionOnViewports(ToggleShowPivotsFunc);
+
+	if (FEngineAnalytics::IsAvailable())
+	{
+		FEngineAnalytics::GetProvider().RecordEvent(TEXT("Editor.Usage.StaticMesh.Toolbar"), TEXT("bDrawPivots"), bDrawPivots ? TEXT("True") : TEXT("False"));
+	}
+}
+
+bool FStaticMeshEditor::IsShowPivotsChecked() const
+{
+	return bDrawPivots;
+}
+
+void FStaticMeshEditor::ToggleShowVertices()
+{
+	bDrawVertices = !bDrawVertices;
+	TFunction<void(FName, TSharedPtr<IEditorViewportLayoutEntity>)> ToggleShowVerticesFunc =
+		[this](FName Name, TSharedPtr<IEditorViewportLayoutEntity> Entity)
+	{
+		TSharedRef<SStaticMeshEditorViewport> StaticMeshEditorViewport = StaticCastSharedRef<SStaticMeshEditorViewport>(Entity->AsWidget());
+		FStaticMeshEditorViewportClient& StaticMeshEditorViewportClient = StaticMeshEditorViewport->GetViewportClient();
+		StaticMeshEditorViewportClient.SetShowVertices(bDrawVertices);
+	};
+
+	ViewportTabContent->PerformActionOnViewports(ToggleShowVerticesFunc);
+
+	if (FEngineAnalytics::IsAvailable())
+	{
+		FEngineAnalytics::GetProvider().RecordEvent(TEXT("Editor.Usage.StaticMesh.Toolbar"), TEXT("bDrawVertices"), bDrawVertices ? TEXT("True") : TEXT("False"));
+	}
+}
+
+bool FStaticMeshEditor::IsShowVerticesChecked() const
+{
+	return bDrawVertices;
+}
+
+void FStaticMeshEditor::ToggleShowGrids()
+{
+	bDrawGrids = !bDrawGrids;
+	TFunction<void(FName, TSharedPtr<IEditorViewportLayoutEntity>)> ToggleShowGridFunc =
+		[this](FName Name, TSharedPtr<IEditorViewportLayoutEntity> Entity)
+	{
+		TSharedRef<SStaticMeshEditorViewport> StaticMeshEditorViewport = StaticCastSharedRef<SStaticMeshEditorViewport>(Entity->AsWidget());
+		FStaticMeshEditorViewportClient& StaticMeshEditorViewportClient = StaticMeshEditorViewport->GetViewportClient();
+		StaticMeshEditorViewportClient.SetShowGrids(bDrawGrids);
+	};
+
+	ViewportTabContent->PerformActionOnViewports(ToggleShowGridFunc);
+
+	if (FEngineAnalytics::IsAvailable())
+	{
+		FEngineAnalytics::GetProvider().RecordEvent(TEXT("Editor.Usage.StaticMesh.Toolbar"), TEXT("bDrawGrids"), bDrawGrids ? TEXT("True") : TEXT("False"));
+	}
+}
+
+bool FStaticMeshEditor::IsShowGridsChecked() const
+{
+	return bDrawGrids;
+}
+
+
+void FStaticMeshEditor::ToggleShowBounds()
+{
+	bDrawBounds = !bDrawBounds;
+	TFunction<void(FName, TSharedPtr<IEditorViewportLayoutEntity>)> ToggleShowBoundsFunc =
+		[this](FName Name, TSharedPtr<IEditorViewportLayoutEntity> Entity)
+	{
+		TSharedRef<SStaticMeshEditorViewport> StaticMeshEditorViewport = StaticCastSharedRef<SStaticMeshEditorViewport>(Entity->AsWidget());
+		FStaticMeshEditorViewportClient& StaticMeshEditorViewportClient = StaticMeshEditorViewport->GetViewportClient();
+		StaticMeshEditorViewportClient.SetShowBounds(bDrawBounds);
+	};
+
+	ViewportTabContent->PerformActionOnViewports(ToggleShowBoundsFunc);
+
+	if (FEngineAnalytics::IsAvailable())
+	{
+		FEngineAnalytics::GetProvider().RecordEvent(TEXT("Editor.Usage.StaticMesh.Toolbar"), TEXT("bDrawBounds"), bDrawBounds ? TEXT("True") : TEXT("False"));
+	}
+}
+
+bool FStaticMeshEditor::IsShowBoundsChecked() const
+{
+	return bDrawBounds;
+}
+
+
+void FStaticMeshEditor::ToggleShowSimpleCollisions()
+{
+	bDrawSimpleCollisions = !bDrawSimpleCollisions;
+	TFunction<void(FName, TSharedPtr<IEditorViewportLayoutEntity>)> ToggleShowSimpleCollisionsFunc =
+		[this](FName Name, TSharedPtr<IEditorViewportLayoutEntity> Entity)
+	{
+		TSharedRef<SStaticMeshEditorViewport> StaticMeshEditorViewport = StaticCastSharedRef<SStaticMeshEditorViewport>(Entity->AsWidget());
+		FStaticMeshEditorViewportClient& StaticMeshEditorViewportClient = StaticMeshEditorViewport->GetViewportClient();
+		StaticMeshEditorViewportClient.SetShowSimpleCollisions(bDrawSimpleCollisions);
+	};
+
+	ViewportTabContent->PerformActionOnViewports(ToggleShowSimpleCollisionsFunc);
+
+	if (FEngineAnalytics::IsAvailable())
+	{
+		FEngineAnalytics::GetProvider().RecordEvent(TEXT("Editor.Usage.StaticMesh.Toolbar"), TEXT("bDrawSimpleCollisions"), bDrawSimpleCollisions ? TEXT("True") : TEXT("False"));
+	}
+}
+
+bool FStaticMeshEditor::IsShowSimpleCollisionsChecked() const
+{
+	return bDrawSimpleCollisions;
+}
+
+
+void FStaticMeshEditor::ToggleShowComplexCollisions()
+{
+	bDrawComplexCollisions = !bDrawComplexCollisions;
+	TFunction<void(FName, TSharedPtr<IEditorViewportLayoutEntity>)> ToggleShowComplexCollisionsFunc =
+		[this](FName Name, TSharedPtr<IEditorViewportLayoutEntity> Entity)
+	{
+		TSharedRef<SStaticMeshEditorViewport> StaticMeshEditorViewport = StaticCastSharedRef<SStaticMeshEditorViewport>(Entity->AsWidget());
+		FStaticMeshEditorViewportClient& StaticMeshEditorViewportClient = StaticMeshEditorViewport->GetViewportClient();
+		StaticMeshEditorViewportClient.SetShowComplexCollisions(bDrawComplexCollisions);
+	};
+
+	ViewportTabContent->PerformActionOnViewports(ToggleShowComplexCollisionsFunc);
+
+	if (FEngineAnalytics::IsAvailable())
+	{
+		FEngineAnalytics::GetProvider().RecordEvent(TEXT("Editor.Usage.StaticMesh.Toolbar"), TEXT("bDrawComplexCollisions"), bDrawComplexCollisions ? TEXT("True") : TEXT("False"));
+	}
+}
+
+bool FStaticMeshEditor::IsShowComplexCollisionsChecked() const
+{
+	return bDrawComplexCollisions;
+}
+
+void FStaticMeshEditor::ToggleShowSockets()
+{
+	bDrawSockets = !bDrawSockets;
+	TFunction<void(FName, TSharedPtr<IEditorViewportLayoutEntity>)> ToggleShowSocketsFunc =
+		[this](FName Name, TSharedPtr<IEditorViewportLayoutEntity> Entity)
+	{
+		TSharedRef<SStaticMeshEditorViewport> StaticMeshEditorViewport = StaticCastSharedRef<SStaticMeshEditorViewport>(Entity->AsWidget());
+		FStaticMeshEditorViewportClient& StaticMeshEditorViewportClient = StaticMeshEditorViewport->GetViewportClient();
+		StaticMeshEditorViewportClient.SetShowPivots(bDrawSockets);
+	};
+
+	ViewportTabContent->PerformActionOnViewports(ToggleShowSocketsFunc);
+
+	if (FEngineAnalytics::IsAvailable())
+	{
+		FEngineAnalytics::GetProvider().RecordEvent(TEXT("Editor.Usage.StaticMesh.Toolbar"), TEXT("bDrawSockets"), bDrawSockets ? TEXT("True") : TEXT("False"));
+	}
+}
+
+bool FStaticMeshEditor::IsShowSocketsChecked() const
+{
+	return bDrawSockets;
+}
+
+void FStaticMeshEditor::ToggleShowWireframes()
+{
+	bDrawWireframes = !bDrawWireframes;
+	TFunction<void(FName, TSharedPtr<IEditorViewportLayoutEntity>)> ToggleShowWireframesFunc =
+		[this](FName Name, TSharedPtr<IEditorViewportLayoutEntity> Entity)
+	{
+		TSharedRef<SStaticMeshEditorViewport> StaticMeshEditorViewport = StaticCastSharedRef<SStaticMeshEditorViewport>(Entity->AsWidget());
+		if (StaticMeshEditorViewport->IsInViewModeWireframeChecked() != bDrawWireframes)
+		{
+			StaticMeshEditorViewport->SetViewModeWireframe();
+		}
+	};
+
+	ViewportTabContent->PerformActionOnViewports(ToggleShowWireframesFunc);
+
+	if (FEngineAnalytics::IsAvailable())
+	{
+		FEngineAnalytics::GetProvider().RecordEvent(TEXT("Editor.Usage.StaticMesh.Toolbar"), TEXT("bDrawWireframes"), bDrawWireframes ? TEXT("True") : TEXT("False"));
+	}
+}
+
+bool FStaticMeshEditor::IsShowWireframesChecked() const
+{
+	return bDrawWireframes;
+}
+
+void FStaticMeshEditor::ToggleShowVertexColors()
+{
+	bDrawVertexColors = !bDrawVertexColors;
+
+	TFunction<void(FName, TSharedPtr<IEditorViewportLayoutEntity>)> ToggleShowVertexColorsFunc =
+		[this](FName Name, TSharedPtr<IEditorViewportLayoutEntity> Entity)
+	{
+		TSharedRef<SStaticMeshEditorViewport> StaticMeshEditorViewport = StaticCastSharedRef<SStaticMeshEditorViewport>(Entity->AsWidget());
+		if (StaticMeshEditorViewport->IsInViewModeVertexColorChecked() != bDrawVertexColors)
+		{
+ 			FStaticMeshEditorViewportClient& EditorViewportClient = StaticMeshEditorViewport->GetViewportClient();
+
+			EditorViewportClient.EngineShowFlags.SetVertexColors(bDrawVertexColors);
+			EditorViewportClient.EngineShowFlags.SetLighting(!bDrawVertexColors);
+			EditorViewportClient.EngineShowFlags.SetIndirectLightingCache(!bDrawVertexColors);
+			EditorViewportClient.EngineShowFlags.SetPostProcessing(!bDrawVertexColors);
+			EditorViewportClient.SetFloorAndEnvironmentVisibility(!bDrawVertexColors);
+			GetStaticMeshComponent()->bDisplayVertexColors = bDrawVertexColors;
+			GetStaticMeshComponent()->MarkRenderStateDirty();
+			StaticMeshEditorViewport->Invalidate();
+		}
+	};
+
+	ViewportTabContent->PerformActionOnViewports(ToggleShowVertexColorsFunc);
+
+	if (FEngineAnalytics::IsAvailable())
+	{
+		FEngineAnalytics::GetProvider().RecordEvent(TEXT("Editor.Usage.StaticMesh.Toolbar"), TEXT("bDrawVertexColors"), bDrawVertexColors ? TEXT("True") : TEXT("False"));
+	}
+}
+
+bool FStaticMeshEditor::IsShowVertexColorsChecked() const
+{
+	return bDrawVertexColors;
+}
+
 void FStaticMeshEditor::RemoveCurrentUVChannel()
 {
 	if (!StaticMesh)
diff --git a/Engine/Source/Editor/StaticMeshEditor/Private/StaticMeshEditor.h b/Engine/Source/Editor/StaticMeshEditor/Private/StaticMeshEditor.h
--- a/Engine/Source/Editor/StaticMeshEditor/Private/StaticMeshEditor.h
+++ b/Engine/Source/Editor/StaticMeshEditor/Private/StaticMeshEditor.h
@@ -30,6 +30,7 @@ class SStaticMeshEditorViewport;
 class UStaticMesh;
 class UStaticMeshComponent;
 class UStaticMeshSocket;
+class FViewportTabContent;
 struct FPropertyChangedEvent;
 struct FTabSpawnerEntry;
 
@@ -40,11 +41,23 @@ class FStaticMeshEditor : public IStaticMeshEditor, public FGCObject, public FEd
 {
 public:
 	FStaticMeshEditor()
-		: StaticMesh(NULL)
+		: StaticMesh(nullptr)
 		, MinPrimSize(0.5f)
 		, OverlapNudge(10.0f)
 		, CurrentViewedUVChannel(0)
 		, SecondaryToolbarEntry(nullptr)
+		, bDrawNormals(false)
+		, bDrawTangents(false)
+		, bDrawBinormals(false)
+		, bDrawPivots(false)
+		, bDrawVertices(false)
+		, bDrawGrids(false)
+		, bDrawBounds(false)
+		, bDrawSimpleCollisions(false)
+		, bDrawComplexCollisions(false)
+		, bDrawSockets(false)
+		, bDrawWireframes(false)
+		, bDrawVertexColors(false)
 	{}
 
 	~FStaticMeshEditor();
@@ -204,7 +217,85 @@ private:
 	TSharedRef<SDockTab> SpawnTab_PreviewSceneSettings(const FSpawnTabArgs& Args);
 	TSharedRef<SDockTab> SpawnTab_SecondaryToolbar(const FSpawnTabArgs& Args);
 
+
+	/** Callback for toggling the normals show flag. */
+	void ToggleShowNormals();
+
+	/** Callback for checking the normals show flag. */
+	bool IsShowNormalsChecked() const;
+
+	/** Callback for toggling the tangents show flag. */
+	void ToggleShowTangents();
+
+	/** Callback for checking the tangents show flag. */
+	bool IsShowTangentsChecked() const;
+
+	/** Callback for toggling the binormals show flag. */
+	void ToggleShowBinormals();
+
+	/** Callback for checking the binormals show flag. */
+	bool IsShowBinormalsChecked() const;
+
+	/** Callback for toggling the pivots show flag. */
+	void ToggleShowPivots();
+
+	/** Callback for checking the pivots show flag. */
+	bool IsShowPivotsChecked() const;
+
+	/** Callback for toggling the vertices show flag. */
+	void ToggleShowVertices();
+
+	/** Callback for checking the vertices show flag. */
+	bool IsShowVerticesChecked() const;
+
+	/** Callback for toggling the grid show flag. */
+	void ToggleShowGrids();
+
+	/** Callback for checking the grid show flag. */
+	bool IsShowGridsChecked() const;
+
+	/** Callback for toggling the bounds show flag. */
+	void ToggleShowBounds();
+
+	/** Callback for checking the bounds show flag. */
+	bool IsShowBoundsChecked() const;
+
+	/** Callback for toggling the simple collisions show flag. */
+	void ToggleShowSimpleCollisions();
+
+	/** Callback for checking the simple collisions show flag. */
+	bool IsShowSimpleCollisionsChecked() const;
+
+	/** Callback for toggling the complex collisions show flag. */
+	void ToggleShowComplexCollisions();
+
+	/** Callback for checking the complex collisions show flag. */
+	bool IsShowComplexCollisionsChecked() const;
+
+	/** Callback for toggling the sockets show flag. */
+	void ToggleShowSockets();
+
+	/** Callback for checking the sockets show flag. */
+	bool IsShowSocketsChecked() const;
+
+	/** Callback for toggling the wireframes show flag. */
+	void ToggleShowWireframes();
+
+	/** Callback for checking the wireframes show flag. */
+	bool IsShowWireframesChecked() const;
+
+	/** Callback for toggling the vertex colors show flag. */
+	void ToggleShowVertexColors();
+
+	/** Callback for checking the vertex colors show flag. */
+	bool IsShowVertexColorsChecked() const;
+
 private:
+
+
+	/** Gets First Viewport from layout. */
+	TSharedPtr<class SStaticMeshEditorViewport> GetStaticMeshViewport() const;
+
 	/** Binds commands associated with the Static Mesh Editor. */
 	void BindCommands();
 
@@ -355,12 +446,12 @@ private:
 	void GenerateSecondaryToolbar();
 
 private:
+	// Tracking the active viewports in this editor.
+	TSharedPtr<class FEditorViewportTabContent> ViewportTabContent;
+
 	/** List of open tool panels; used to ensure only one exists at any one time */
 	TMap< FName, TWeakPtr<class SDockableTab> > SpawnedToolPanels;
 
-	/** Preview Viewport widget */
-	TSharedPtr<class SStaticMeshEditorViewport> Viewport;
-
 	/** Property View */
 	TSharedPtr<class IDetailsView> StaticMeshDetailsView;
 
@@ -438,4 +529,21 @@ private:
 
 	/** The text display name to override the default display name of the secondary toolbar*/
 	FText SecondaryToolbarDisplayName;
+	
+	/** Storage for our viewport creation function that will be passed to the viewport layout system*/
+	TFunction<TSharedRef<SEditorViewport>(void)> MakeViewportFunc;
+
+	/** Toolbar toggles */
+	bool bDrawNormals;
+	bool bDrawTangents;
+	bool bDrawBinormals;
+	bool bDrawPivots;
+	bool bDrawVertices;
+	bool bDrawGrids;
+	bool bDrawBounds;
+	bool bDrawSimpleCollisions;
+	bool bDrawComplexCollisions;
+	bool bDrawSockets;
+	bool bDrawWireframes;
+	bool bDrawVertexColors;
 };
diff --git a/Engine/Source/Editor/StaticMeshEditor/Private/StaticMeshEditorActions.cpp b/Engine/Source/Editor/StaticMeshEditor/Private/StaticMeshEditorActions.cpp
--- a/Engine/Source/Editor/StaticMeshEditor/Private/StaticMeshEditorActions.cpp
+++ b/Engine/Source/Editor/StaticMeshEditor/Private/StaticMeshEditorActions.cpp
@@ -22,6 +22,19 @@ void FStaticMeshEditorCommands::RegisterCommands()
 	UI_COMMAND( ReimportAllMesh, "Reimport Base Mesh + LODs", "Reimport the base mesh and all the custom LODs.", EUserInterfaceActionType::Button, FInputChord());
 	UI_COMMAND( ReimportAllMeshWithNewFile, "Reimport Base Mesh + LODs With New File", "Reimport the base mesh using a new source file and all the custom LODs (No new source file for the custom LODs).", EUserInterfaceActionType::Button, FInputChord() );
 
+	UI_COMMAND( ToggleShowNormals, "Normals", "Toggles display of vertex normals in the Preview Pane.", EUserInterfaceActionType::ToggleButton, FInputChord());
+	UI_COMMAND( ToggleShowTangents, "Tangents", "Toggles display of vertex tangents in the Preview Pane.", EUserInterfaceActionType::ToggleButton, FInputChord());
+	UI_COMMAND( ToggleShowBinormals, "Binormals", "Toggles display of vertex binormals (orthogonal vector to normal and tangent) in the Preview Pane.", EUserInterfaceActionType::ToggleButton, FInputChord());
+	UI_COMMAND( ToggleShowVertices, "Vertices", "Toggles display of vertices in the Preview Pane.", EUserInterfaceActionType::ToggleButton, FInputChord());
+	UI_COMMAND( ToggleShowPivots, "Show Pivot", "Display the pivot location of the static mesh.", EUserInterfaceActionType::ToggleButton, FInputChord());
+	UI_COMMAND( ToggleShowGrids, "Grid", "Displays the viewport grid.", EUserInterfaceActionType::ToggleButton, FInputChord());
+	UI_COMMAND( ToggleShowBounds, "Bounds", "Toggles display of the bounds of the static mesh.", EUserInterfaceActionType::ToggleButton, FInputChord());
+	UI_COMMAND( ToggleShowSimpleCollisions, "Simple Collision", "Toggles display of the simplified collision mesh of the static mesh, if present.", EUserInterfaceActionType::ToggleButton, FInputChord());
+	UI_COMMAND( ToggleShowComplexCollisions, "Complex Collision", "Toggles display of the complex collision for this static mesh.", EUserInterfaceActionType::ToggleButton, FInputChord());
+	UI_COMMAND( ToggleShowSockets, "Sockets", "Displays the static mesh sockets.", EUserInterfaceActionType::ToggleButton, FInputChord());
+	UI_COMMAND( ToggleShowWireframes, "Wireframe", "Toggles the viewmode of the Preview Pane between a lit view and a wireframe view.", EUserInterfaceActionType::ToggleButton, FInputChord());
+	UI_COMMAND( ToggleShowVertexColors, "Vert Colors", "Toggles vertex colors.", EUserInterfaceActionType::ToggleButton, FInputChord());
+
 	UI_COMMAND( SetShowNormals, "Normals", "Toggles display of vertex normals in the Preview Pane.", EUserInterfaceActionType::ToggleButton, FInputChord() );
 	UI_COMMAND( SetShowTangents, "Tangents", "Toggles display of vertex tangents in the Preview Pane.", EUserInterfaceActionType::ToggleButton, FInputChord() );
 	UI_COMMAND( SetShowBinormals, "Binormals", "Toggles display of vertex binormals (orthogonal vector to normal and tangent) in the Preview Pane.", EUserInterfaceActionType::ToggleButton, FInputChord() );
diff --git a/Engine/Source/Editor/StaticMeshEditor/Private/StaticMeshEditorViewportClient.cpp b/Engine/Source/Editor/StaticMeshEditor/Private/StaticMeshEditorViewportClient.cpp
--- a/Engine/Source/Editor/StaticMeshEditor/Private/StaticMeshEditorViewportClient.cpp
+++ b/Engine/Source/Editor/StaticMeshEditor/Private/StaticMeshEditorViewportClient.cpp
@@ -1405,36 +1405,146 @@ void FStaticMeshEditorViewportClient::SetPreviewMesh(UStaticMesh* InStaticMesh,
 	}
 }
 
-void FStaticMeshEditorViewportClient::ToggleDrawUVOverlay()
+void FStaticMeshEditorViewportClient::ToggleShowNormals()
 {
-	SetDrawUVOverlay(!bDrawUVs);
+	bDrawNormals = !bDrawNormals;
+	if (FEngineAnalytics::IsAvailable())
+	{
+		FEngineAnalytics::GetProvider().RecordEvent(TEXT("Editor.Usage.StaticMesh.Toolbar"), TEXT("bDrawNormals"), bDrawNormals ? TEXT("True") : TEXT("False"));
+	}
+	Invalidate();
 }
 
-void FStaticMeshEditorViewportClient::SetDrawUVOverlay(bool bShouldDraw)
+
+void FStaticMeshEditorViewportClient::SetShowNormals(bool bShowOn)
 {
-	bDrawUVs = bShouldDraw;
+	bDrawNormals = bShowOn;
 	if (FEngineAnalytics::IsAvailable())
 	{
-		FEngineAnalytics::GetProvider().RecordEvent(TEXT("Editor.Usage.StaticMesh.Toolbar"), TEXT("bDrawUVs"), bDrawUVs ? TEXT("True") : TEXT("False"));
+		FEngineAnalytics::GetProvider().RecordEvent(TEXT("Editor.Usage.StaticMesh.Toolbar"), TEXT("bDrawNormals"), bDrawNormals ? TEXT("True") : TEXT("False"));
 	}
 	Invalidate();
 }
 
-bool FStaticMeshEditorViewportClient::IsDrawUVOverlayChecked() const
+void FStaticMeshEditorViewportClient::SetShowTangents(bool bShowOn)
 {
-	return bDrawUVs;
+	bDrawTangents = bShowOn;
+	if (FEngineAnalytics::IsAvailable())
+	{
+		FEngineAnalytics::GetProvider().RecordEvent(TEXT("Editor.Usage.StaticMesh.Toolbar"), TEXT("bDrawTangents"), bDrawTangents ? TEXT("True") : TEXT("False"));
+	}
+	Invalidate();
 }
 
-void FStaticMeshEditorViewportClient::ToggleShowNormals()
+void FStaticMeshEditorViewportClient::SetShowBinormals(bool bShowOn)
 {
-	bDrawNormals = !bDrawNormals;
+	bDrawBinormals = bShowOn;
 	if (FEngineAnalytics::IsAvailable())
 	{
-		FEngineAnalytics::GetProvider().RecordEvent(TEXT("Editor.Usage.StaticMesh.Toolbar"), TEXT("bDrawNormals"), bDrawNormals ? TEXT("True") : TEXT("False"));
+		FEngineAnalytics::GetProvider().RecordEvent(TEXT("Editor.Usage.StaticMesh.Toolbar"), TEXT("bDrawBinormals"), bDrawBinormals ? TEXT("True") : TEXT("False"));
 	}
 	Invalidate();
 }
 
+void FStaticMeshEditorViewportClient::SetShowSimpleCollisions(bool bShowOn)
+{
+	// ToggleShowSimpleCollision() does more that just flipping a flag so we allow it to do its thing if needed.
+	if (bShowSimpleCollision != bShowOn)
+	{
+		ToggleShowSimpleCollision();
+		if (FEngineAnalytics::IsAvailable())
+		{
+			FEngineAnalytics::GetProvider().RecordEvent(TEXT("Editor.Usage.StaticMesh.Toolbar"), TEXT("bShowComplexCollision"), bShowPivot ? TEXT("True") : TEXT("False"));
+		}
+		Invalidate();
+	}
+
+
+	bShowSimpleCollision = bShowOn;
+	if (FEngineAnalytics::IsAvailable())
+	{
+		FEngineAnalytics::GetProvider().RecordEvent(TEXT("Editor.Usage.StaticMesh.Toolbar"), TEXT("bShowSimpleCollision"), bShowPivot ? TEXT("True") : TEXT("False"));
+	}
+	Invalidate();
+}
+
+void FStaticMeshEditorViewportClient::SetShowComplexCollisions(bool bShowOn)
+{
+	// ToggleShowComplexCollision() does more that just flipping a flag so we allow it to do its thing if needed.
+	if (bShowComplexCollision != bShowOn)
+	{
+		ToggleShowComplexCollision();
+		if (FEngineAnalytics::IsAvailable())
+		{
+			FEngineAnalytics::GetProvider().RecordEvent(TEXT("Editor.Usage.StaticMesh.Toolbar"), TEXT("bShowComplexCollision"), bShowPivot ? TEXT("True") : TEXT("False"));
+		}
+		Invalidate();
+	}
+}
+
+void FStaticMeshEditorViewportClient::SetShowPivots(bool bShowOn)
+{
+	bShowPivot = bShowOn;
+	if (FEngineAnalytics::IsAvailable())
+	{
+		FEngineAnalytics::GetProvider().RecordEvent(TEXT("Editor.Usage.StaticMesh.Toolbar"), TEXT("bShowPivot"), bShowPivot ? TEXT("True") : TEXT("False"));
+	}
+	Invalidate();
+}
+
+void FStaticMeshEditorViewportClient::SetShowGrids(bool bShowOn)
+{
+	DrawHelper.bDrawGrid = bShowOn;
+	if (FEngineAnalytics::IsAvailable())
+	{
+		FEngineAnalytics::GetProvider().RecordEvent(TEXT("Editor.Usage.StaticMesh.Toolbar"), TEXT("bDrawGrids"), bDrawVertices ? TEXT("True") : TEXT("False"));
+	}
+	Invalidate();
+}
+
+void FStaticMeshEditorViewportClient::SetShowVertices(bool bShowOn)
+{
+	bDrawVertices = bShowOn;
+	if (FEngineAnalytics::IsAvailable())
+	{
+		FEngineAnalytics::GetProvider().RecordEvent(TEXT("Editor.Usage.StaticMesh.Toolbar"), TEXT("bDrawVertices"), bDrawVertices ? TEXT("True") : TEXT("False"));
+	}
+	Invalidate();
+}
+
+
+void FStaticMeshEditorViewportClient::SetShowWireframes(bool bShowOn)
+{
+
+}
+
+
+void FStaticMeshEditorViewportClient::SetShowVertexColors(bool bShowOn)
+{
+
+}
+
+
+void FStaticMeshEditorViewportClient::ToggleDrawUVOverlay()
+{
+	SetDrawUVOverlay(!bDrawUVs);
+}
+
+void FStaticMeshEditorViewportClient::SetDrawUVOverlay(bool bShouldDraw)
+{
+	bDrawUVs = bShouldDraw;
+	if (FEngineAnalytics::IsAvailable())
+	{
+		FEngineAnalytics::GetProvider().RecordEvent(TEXT("Editor.Usage.StaticMesh.Toolbar"), TEXT("bDrawUVs"), bDrawUVs ? TEXT("True") : TEXT("False"));
+	}
+	Invalidate();
+}
+
+bool FStaticMeshEditorViewportClient::IsDrawUVOverlayChecked() const
+{
+	return bDrawUVs;
+}
+
 bool FStaticMeshEditorViewportClient::IsShowNormalsChecked() const
 {
 	return bDrawNormals;
diff --git a/Engine/Source/Editor/StaticMeshEditor/Private/StaticMeshEditorViewportClient.h b/Engine/Source/Editor/StaticMeshEditor/Private/StaticMeshEditorViewportClient.h
--- a/Engine/Source/Editor/StaticMeshEditor/Private/StaticMeshEditorViewportClient.h
+++ b/Engine/Source/Editor/StaticMeshEditor/Private/StaticMeshEditorViewportClient.h
@@ -81,6 +81,37 @@ public:
 	/** Callback for checking the UV overlay show flag. */
 	bool IsDrawUVOverlayChecked() const;
 
+
+	/** Set the normals show flag. */
+	void SetShowNormals(bool bShowOn);
+
+	/** Set the tangents show flag. */
+	void SetShowTangents(bool bShowOn);
+
+	/** Set the binormals show flag. */
+	void SetShowBinormals(bool bShowOn);
+
+	/** Set the simple collisions show flag. */
+	void SetShowSimpleCollisions(bool bShowOn);
+
+	/** Set the complex collisions show flag. */
+	void SetShowComplexCollisions(bool bShowOn);
+
+	/** Set the pivots show flag. */
+	void SetShowPivots(bool bShowOn);
+
+	/** Set the grids show flag. */
+	void SetShowGrids(bool bShowOn);
+
+	/** Set the vertices show flag. */
+	void SetShowVertices(bool bShowOn);
+
+	/** Set the wireframes show flag. */
+	void SetShowWireframes(bool bShowOn);
+
+	/** Set the vertex colors show flag. */
+	void SetShowVertexColors(bool bShowOn);
+
 	/** Callback for toggling the normals show flag. */
 	void ToggleShowNormals();
 	
diff --git a/Engine/Source/Editor/StaticMeshEditor/Public/StaticMeshEditorActions.h b/Engine/Source/Editor/StaticMeshEditor/Public/StaticMeshEditorActions.h
--- a/Engine/Source/Editor/StaticMeshEditor/Public/StaticMeshEditorActions.h
+++ b/Engine/Source/Editor/StaticMeshEditor/Public/StaticMeshEditorActions.h
@@ -47,6 +47,20 @@ public:
 	TSharedPtr< FUICommandInfo > ReimportAllMesh;
 	TSharedPtr< FUICommandInfo > ReimportAllMeshWithNewFile;
 
+	// toolbar commands
+	TSharedPtr< FUICommandInfo > ToggleShowNormals;
+	TSharedPtr< FUICommandInfo > ToggleShowTangents;
+	TSharedPtr< FUICommandInfo > ToggleShowBinormals;
+	TSharedPtr< FUICommandInfo > ToggleShowPivots;
+	TSharedPtr< FUICommandInfo > ToggleShowVertices;
+	TSharedPtr< FUICommandInfo > ToggleShowGrids;
+	TSharedPtr< FUICommandInfo > ToggleShowBounds;
+	TSharedPtr< FUICommandInfo > ToggleShowSimpleCollisions;
+	TSharedPtr< FUICommandInfo > ToggleShowComplexCollisions;
+	TSharedPtr< FUICommandInfo > ToggleShowSockets;
+	TSharedPtr< FUICommandInfo > ToggleShowWireframes;
+	TSharedPtr< FUICommandInfo > ToggleShowVertexColors;
+
 	// View Menu Commands
 	TSharedPtr< FUICommandInfo > SetShowNormals;
 	TSharedPtr< FUICommandInfo > SetShowTangents;
