CIS fix #rb none

diff --git a/Engine/Source/Editor/StaticMeshEditor/Private/StaticMeshEditor.h b/Engine/Source/Editor/StaticMeshEditor/Private/StaticMeshEditor.h
--- a/Engine/Source/Editor/StaticMeshEditor/Private/StaticMeshEditor.h
+++ b/Engine/Source/Editor/StaticMeshEditor/Private/StaticMeshEditor.h
@@ -10,6 +10,7 @@
 #include "IStaticMeshEditor.h"
 #include "ISocketManager.h"
 #include "TickableEditorObject.h"
+#include "SEditorViewport.h"
 
 // Set USE_ASYNC_DECOMP to zero to go back to the fully synchronous; blocking version of V-HACD
 #ifndef USE_ASYNC_DECOMP
