Hot fix for a crash when reseting the import asset data on a static mesh via the static mesh editor.
Removed the ui to reset the asset import data on a static mesh.

diff --git a/Engine/Source/Editor/StaticMeshEditor/Private/StaticMeshEditorTools.cpp b/Engine/Source/Editor/StaticMeshEditor/Private/StaticMeshEditorTools.cpp
--- a/Engine/Source/Editor/StaticMeshEditor/Private/StaticMeshEditorTools.cpp
+++ b/Engine/Source/Editor/StaticMeshEditor/Private/StaticMeshEditorTools.cpp
@@ -32,22 +32,22 @@
 #include "SPerPlatformPropertiesWidget.h"
 #include "PlatformInfo.h"
 
-#include "Runtime/Analytics/Analytics/Public/Interfaces/IAnalyticsProvider.h"
-#include "EngineAnalytics.h"
-#include "Widgets/Input/STextComboBox.h"
-#include "ScopedTransaction.h"
-#include "JsonObjectConverter.h"
-#include "Engine/SkeletalMesh.h"
-#include "IMeshReductionManagerModule.h"
-#include "HAL/PlatformApplicationMisc.h"
-#include "Widgets/Input/SFilePathPicker.h"
+#include "ContentStreaming.h"
 #include "EditorDirectories.h"
 #include "EditorFramework/AssetImportData.h"
+#include "Engine/SkeletalMesh.h"
+#include "EngineAnalytics.h"
 #include "Factories/FbxStaticMeshImportData.h"
-#include "Interfaces/ITargetPlatformManagerModule.h"
+#include "HAL/PlatformApplicationMisc.h"
+#include "IMeshReductionManagerModule.h"
 #include "Interfaces/ITargetPlatform.h"
-#include "ContentStreaming.h"
-
+#include "Interfaces/ITargetPlatformManagerModule.h"
+#include "JsonObjectConverter.h"
+#include "Runtime/Analytics/Analytics/Public/Interfaces/IAnalyticsProvider.h"
+#include "ScopedTransaction.h"
+#include "UObject/UObjectGlobals.h"
+#include "Widgets/Input/SFilePathPicker.h"
+#include "Widgets/Input/STextComboBox.h"
 
 const uint32 MaxHullCount = 64;
 const uint32 MinHullCount = 2;
@@ -93,9 +93,20 @@ void FStaticMeshDetails::CustomizeDetails( class IDetailLayoutBuilder& DetailBui
 	IDetailCategoryBuilder& ImportSettingsCategory = DetailBuilder.EditCategory("ImportSettings");
 
 	TSharedRef<IPropertyHandle> ImportSettings = DetailBuilder.GetProperty(GET_MEMBER_NAME_CHECKED(UStaticMesh, AssetImportData));
+
+	/**
+	 * Hotfix 4.24.1 ( Continues the old and maybe invalid assumption that AssetImportData shouldn't be able to be null )
+	 * This will be removed in 4.25
+	 */
+	if( !StaticMeshEditor.GetStaticMesh()->AssetImportData )
+	{
+		StaticMeshEditor.GetStaticMesh()->AssetImportData = NewObject<UAssetImportData>(StaticMeshEditor.GetStaticMesh(), TEXT("AssetImportData"), RF_Transactional);
+	}
+
 	if (!StaticMeshEditor.GetStaticMesh() || !StaticMeshEditor.GetStaticMesh()->AssetImportData->IsA<UFbxStaticMeshImportData>())
 	{
-		// Hide the ability to change the import settings object
+		ImportSettings->MarkResetToDefaultCustomized();
+
 		IDetailPropertyRow& Row = ImportSettingsCategory.AddProperty(ImportSettings);
 		Row.CustomWidget(true)
 			.NameContent()
