Added transaction to StaticMesh RemoveUVChannel

diff --git a/Engine/Source/Editor/StaticMeshEditor/Private/StaticMeshEditor.cpp b/Engine/Source/Editor/StaticMeshEditor/Private/StaticMeshEditor.cpp
--- a/Engine/Source/Editor/StaticMeshEditor/Private/StaticMeshEditor.cpp
+++ b/Engine/Source/Editor/StaticMeshEditor/Private/StaticMeshEditor.cpp
@@ -16,6 +16,7 @@
 #include "UnrealEdGlobals.h"
 #include "StaticMeshEditorModule.h"
 #include "AI/Navigation/NavCollisionBase.h"
+#include "ScopedTransaction.h"
 
 #include "SStaticMeshEditorViewport.h"
 #include "PropertyEditorModule.h"
@@ -2704,6 +2705,7 @@ void FStaticMeshEditor::RemoveCurrentUVChannel()
 			}
 		}
 
+		const FScopedTransaction Transaction(LOCTEXT("RemoveUVChannel", "Remove UV Channel"));
 		if (StaticMesh->RemoveUVChannel(LODIndex, UVChannelIndex))
 		{
 			RefreshTool();
