Fix crash in static mesh editor after copy and paste primitives

diff --git a/Engine/Source/Editor/StaticMeshEditor/Private/StaticMeshEditor.cpp b/Engine/Source/Editor/StaticMeshEditor/Private/StaticMeshEditor.cpp
--- a/Engine/Source/Editor/StaticMeshEditor/Private/StaticMeshEditor.cpp
+++ b/Engine/Source/Editor/StaticMeshEditor/Private/StaticMeshEditor.cpp
@@ -2250,6 +2250,7 @@ void FStaticMeshEditor::NotifyPostChange( const FPropertyChangedEvent& PropertyC
 	if(StaticMesh && StaticMesh->BodySetup)
 	{
 		StaticMesh->BodySetup->CreatePhysicsMeshes();
+		RemoveInvalidPrims();
 
 		if (GET_MEMBER_NAME_CHECKED(UStaticMesh, LODGroup) == PropertyChangedEvent.GetPropertyName())
 		{
@@ -2267,7 +2268,10 @@ void FStaticMeshEditor::NotifyPostChange( const FPropertyChangedEvent& PropertyC
 				}
 			}
 		}
-
+	}
+	else
+	{
+		RemoveInvalidPrims();
 	}
 }
 
