Added PhysicalMaterialMasks show flag and visualization in level editor and static mesh editor.

diff --git a/Engine/Config/BaseEngine.ini b/Engine/Config/BaseEngine.ini
--- a/Engine/Config/BaseEngine.ini
+++ b/Engine/Config/BaseEngine.ini
@@ -110,6 +110,7 @@ VertexColorViewModeMaterialName_AlphaAsColor=/Engine/EngineDebugMaterials/Vertex
 VertexColorViewModeMaterialName_RedOnly=/Engine/EngineDebugMaterials/VertexColorViewMode_RedOnly.VertexColorViewMode_RedOnly
 VertexColorViewModeMaterialName_GreenOnly=/Engine/EngineDebugMaterials/VertexColorViewMode_GreenOnly.VertexColorViewMode_GreenOnly
 VertexColorViewModeMaterialName_BlueOnly=/Engine/EngineDebugMaterials/VertexColorViewMode_BlueOnly.VertexColorViewMode_BlueOnly
+PhysicalMaterialMaskMaterialName=/Engine/EngineDebugMaterials/PhysicalMaterialMaskMaterial.PhysicalMaterialMaskMaterial
 BoneWeightMaterialName=/Engine/EngineDebugMaterials/BoneWeightMaterial.BoneWeightMaterial
 ClothPaintMaterialName=/Engine/EngineDebugMaterials/ClothMaterial.ClothMaterial
 ClothPaintMaterialWireframeName=/Engine/EngineDebugMaterials/ClothMaterial_WF.ClothMaterial_WF
diff --git a/Engine/Source/Editor/StaticMeshEditor/Private/SStaticMeshEditorViewport.cpp b/Engine/Source/Editor/StaticMeshEditor/Private/SStaticMeshEditorViewport.cpp
--- a/Engine/Source/Editor/StaticMeshEditor/Private/SStaticMeshEditorViewport.cpp
+++ b/Engine/Source/Editor/StaticMeshEditor/Private/SStaticMeshEditorViewport.cpp
@@ -306,10 +306,15 @@ void SStaticMeshEditorViewport::UpdatePreviewMesh(UStaticMesh* InStaticMesh, boo
 
 	EditorViewportClient->SetPreviewMesh(InStaticMesh, PreviewMeshComponent, bResetCamera);
 
-	if (EditorViewportClient->EngineShowFlags.VertexColors)
+	if (EditorViewportClient->EngineShowFlags.PhysicalMaterialMasks)
+	{
+		//Reapply the physical material masks mode on the newly set static mesh.
+		SetViewModePhysicalMaterialMasksImplementation(true);
+	}
+	else if (EditorViewportClient->EngineShowFlags.VertexColors)
 	{
 		//Reapply the vertex color mode on the newly set static mesh.
-		SetViewModeVertexColorImplemetation(true);
+		SetViewModeVertexColorImplementation(true);
 	}
 
 	PreviewMeshComponent->SelectionOverrideDelegate = UPrimitiveComponent::FSelectionOverride::CreateRaw(this, &SStaticMeshEditorViewport::PreviewComponentSelectionOverride);
@@ -352,7 +357,7 @@ bool SStaticMeshEditorViewport::IsInViewModeWireframeChecked() const
 
 void SStaticMeshEditorViewport::SetViewModeVertexColor()
 {
-	SetViewModeVertexColorImplemetation(!EditorViewportClient->EngineShowFlags.VertexColors);
+	SetViewModeVertexColorImplementation(!EditorViewportClient->EngineShowFlags.VertexColors);
 
 	if (FEngineAnalytics::IsAvailable())
 	{
@@ -360,7 +365,21 @@ void SStaticMeshEditorViewport::SetViewModeVertexColor()
 	}
 }
 
-void SStaticMeshEditorViewport::SetViewModeVertexColorImplemetation(bool bValue)
+void SStaticMeshEditorViewport::SetViewModeVertexColorImplementation(bool bValue)
+{
+	SetViewModeVertexColorSubImplementation(bValue);
+
+	// Disable physical material masks, if enabling vertex color.
+	if (bValue)
+	{
+		SetViewModePhysicalMaterialMasksSubImplementation(false);
+	}
+
+	StaticMeshEditorPtr.Pin()->GetStaticMeshComponent()->MarkRenderStateDirty();
+	SceneViewport->Invalidate();
+}
+
+void SStaticMeshEditorViewport::SetViewModeVertexColorSubImplementation(bool bValue)
 {
 	EditorViewportClient->EngineShowFlags.SetVertexColors(bValue);
 	EditorViewportClient->EngineShowFlags.SetLighting(!bValue);
@@ -368,8 +387,6 @@ void SStaticMeshEditorViewport::SetViewModeVertexColorImplemetation(bool bValue)
 	EditorViewportClient->EngineShowFlags.SetPostProcessing(!bValue);
 	EditorViewportClient->SetFloorAndEnvironmentVisibility(!bValue);
 	StaticMeshEditorPtr.Pin()->GetStaticMeshComponent()->bDisplayVertexColors = bValue;
-	StaticMeshEditorPtr.Pin()->GetStaticMeshComponent()->MarkRenderStateDirty();
-	SceneViewport->Invalidate();
 }
 
 bool SStaticMeshEditorViewport::IsInViewModeVertexColorChecked() const
@@ -377,6 +394,42 @@ bool SStaticMeshEditorViewport::IsInViewModeVertexColorChecked() const
 	return EditorViewportClient->EngineShowFlags.VertexColors;
 }
 
+void SStaticMeshEditorViewport::SetViewModePhysicalMaterialMasks()
+{
+	SetViewModePhysicalMaterialMasksImplementation(!EditorViewportClient->EngineShowFlags.PhysicalMaterialMasks);
+
+	if (FEngineAnalytics::IsAvailable())
+	{
+		FEngineAnalytics::GetProvider().RecordEvent(TEXT("Editor.Usage.StaticMesh.Toolbar"), FAnalyticsEventAttribute(TEXT("PhysicalMaterialMasks"), static_cast<int>(EditorViewportClient->EngineShowFlags.PhysicalMaterialMasks)));
+	}
+}
+
+void SStaticMeshEditorViewport::SetViewModePhysicalMaterialMasksImplementation(bool bValue)
+{
+	SetViewModePhysicalMaterialMasksSubImplementation(bValue);
+
+	// Disable vertex color, if enabling physical material masks.
+	if (bValue)
+	{
+		SetViewModeVertexColorSubImplementation(false);
+	}
+
+	StaticMeshEditorPtr.Pin()->GetStaticMeshComponent()->MarkRenderStateDirty();
+	SceneViewport->Invalidate();
+}
+
+void SStaticMeshEditorViewport::SetViewModePhysicalMaterialMasksSubImplementation(bool bValue)
+{
+	EditorViewportClient->EngineShowFlags.SetPhysicalMaterialMasks(bValue);
+	StaticMeshEditorPtr.Pin()->GetStaticMeshComponent()->bDisplayPhysicalMaterialMasks = bValue;
+}
+
+bool SStaticMeshEditorViewport::IsInViewModePhysicalMaterialMasksChecked() const
+{
+	return EditorViewportClient->EngineShowFlags.PhysicalMaterialMasks;
+}
+
+
 void SStaticMeshEditorViewport::ForceLODLevel(int32 InForcedLOD)
 {
 	PreviewMeshComponent->ForcedLodModel = InForcedLOD;
@@ -486,6 +539,12 @@ void SStaticMeshEditorViewport::BindCommands()
 		FCanExecuteAction(),
 		FIsActionChecked::CreateSP( this, &SStaticMeshEditorViewport::IsInViewModeVertexColorChecked ) );
 
+	CommandList->MapAction(
+		Commands.SetShowPhysicalMaterialMasks,
+		FExecuteAction::CreateSP(this, &SStaticMeshEditorViewport::SetViewModePhysicalMaterialMasks),
+		FCanExecuteAction(),
+		FIsActionChecked::CreateSP(this, &SStaticMeshEditorViewport::IsInViewModePhysicalMaterialMasksChecked));
+
 	CommandList->MapAction(
 		Commands.ResetCamera,
 		FExecuteAction::CreateSP( EditorViewportClientRef, &FStaticMeshEditorViewportClient::ResetCamera ) );
diff --git a/Engine/Source/Editor/StaticMeshEditor/Private/SStaticMeshEditorViewport.h b/Engine/Source/Editor/StaticMeshEditor/Private/SStaticMeshEditorViewport.h
--- a/Engine/Source/Editor/StaticMeshEditor/Private/SStaticMeshEditorViewport.h
+++ b/Engine/Source/Editor/StaticMeshEditor/Private/SStaticMeshEditorViewport.h
@@ -150,7 +150,22 @@ private:
 	bool IsVisible() const override;
 
 	/** Implementation of the SetViewModeVertexColor, used to apply the vertex color show flag. */
-	void SetViewModeVertexColorImplemetation(bool bValue);
+	void SetViewModeVertexColorImplementation(bool bValue);
+
+	/** Applies the vertex color show flag only. */
+	void SetViewModeVertexColorSubImplementation(bool bValue);
+
+	/** Callback for toggling the physical material mask show flag. */
+	void SetViewModePhysicalMaterialMasks();
+
+	/** Implementation of the SetViewModeVertexColor, used to apply the physical material mask show flag. */
+	void SetViewModePhysicalMaterialMasksImplementation(bool bValue);
+
+	/** Applies the physical material mask show flag only. */
+	void SetViewModePhysicalMaterialMasksSubImplementation(bool bValue);
+
+	/** Callback for checking the physical material mask show flag. */
+	bool IsInViewModePhysicalMaterialMasksChecked() const;
 
 	/** Callback for toggling the realtime preview flag. */
 	void SetRealtimePreview();
diff --git a/Engine/Source/Editor/StaticMeshEditor/Private/SStaticMeshEditorViewportToolBar.cpp b/Engine/Source/Editor/StaticMeshEditor/Private/SStaticMeshEditorViewportToolBar.cpp
--- a/Engine/Source/Editor/StaticMeshEditor/Private/SStaticMeshEditorViewportToolBar.cpp
+++ b/Engine/Source/Editor/StaticMeshEditor/Private/SStaticMeshEditorViewportToolBar.cpp
@@ -53,6 +53,7 @@ TSharedRef<SWidget> SStaticMeshEditorViewportToolbar::GenerateShowMenu() const
 		ShowMenuBuilder.AddMenuEntry(Commands.SetShowBounds);
 		ShowMenuBuilder.AddMenuEntry(Commands.SetShowSimpleCollision);
 		ShowMenuBuilder.AddMenuEntry(Commands.SetShowComplexCollision);
+		ShowMenuBuilder.AddMenuEntry(Commands.SetShowPhysicalMaterialMasks);
 
 		ShowMenuBuilder.AddMenuSeparator();
 
diff --git a/Engine/Source/Editor/StaticMeshEditor/Private/StaticMeshAutomationTests.cpp b/Engine/Source/Editor/StaticMeshEditor/Private/StaticMeshAutomationTests.cpp
--- a/Engine/Source/Editor/StaticMeshEditor/Private/StaticMeshAutomationTests.cpp
+++ b/Engine/Source/Editor/StaticMeshEditor/Private/StaticMeshAutomationTests.cpp
@@ -100,6 +100,7 @@ namespace EditorViewButtonHelper
 				if (!AutomationParameters.ViewportClient->EngineShowFlags.VertexColors)
 				{
 					AutomationParameters.ViewportClient->EngineShowFlags.SetVertexColors(true);
+					AutomationParameters.ViewportClient->EngineShowFlags.SetPhysicalMaterialMasks(false);
 					AutomationParameters.ViewportClient->EngineShowFlags.SetLighting(false);
 					AutomationParameters.ViewportClient->EngineShowFlags.SetIndirectLightingCache(false);
 				}
diff --git a/Engine/Source/Editor/StaticMeshEditor/Private/StaticMeshEditorActions.cpp b/Engine/Source/Editor/StaticMeshEditor/Private/StaticMeshEditorActions.cpp
--- a/Engine/Source/Editor/StaticMeshEditor/Private/StaticMeshEditorActions.cpp
+++ b/Engine/Source/Editor/StaticMeshEditor/Private/StaticMeshEditorActions.cpp
@@ -8,6 +8,7 @@ void FStaticMeshEditorCommands::RegisterCommands()
 {
 	UI_COMMAND( SetShowWireframe, "Wireframe", "Toggles the viewmode of the Preview Pane between a lit view and a wireframe view.", EUserInterfaceActionType::ToggleButton, FInputChord() );
 	UI_COMMAND( SetShowVertexColor, "Vert Colors", "Toggles vertex colors.", EUserInterfaceActionType::ToggleButton, FInputChord() );
+	UI_COMMAND( SetShowPhysicalMaterialMasks, "Physical Material Masks", "Toggles physical material masks.", EUserInterfaceActionType::ToggleButton, FInputChord());
 	UI_COMMAND( SetDrawUVs, "UV", "Toggles display of the static mesh's UVs for the specified channel.", EUserInterfaceActionType::ToggleButton, FInputChord() );
 	UI_COMMAND( SetShowGrid, "Grid", "Displays the viewport grid.", EUserInterfaceActionType::ToggleButton, FInputChord() );
 	UI_COMMAND( SetShowBounds, "Bounds", "Toggles display of the bounds of the static mesh.", EUserInterfaceActionType::ToggleButton, FInputChord() );
diff --git a/Engine/Source/Editor/StaticMeshEditor/Public/StaticMeshEditorActions.h b/Engine/Source/Editor/StaticMeshEditor/Public/StaticMeshEditorActions.h
--- a/Engine/Source/Editor/StaticMeshEditor/Public/StaticMeshEditorActions.h
+++ b/Engine/Source/Editor/StaticMeshEditor/Public/StaticMeshEditorActions.h
@@ -32,6 +32,7 @@ public:
 	/**  */
 	TSharedPtr< FUICommandInfo > SetShowWireframe;
 	TSharedPtr< FUICommandInfo > SetShowVertexColor;
+	TSharedPtr< FUICommandInfo > SetShowPhysicalMaterialMasks;
 	TSharedPtr< FUICommandInfo > SetDrawUVs;
 	TSharedPtr< FUICommandInfo > SetShowGrid;
 	TSharedPtr< FUICommandInfo > SetShowBounds;
diff --git a/Engine/Source/Runtime/Engine/Classes/Components/StaticMeshComponent.h b/Engine/Source/Runtime/Engine/Classes/Components/StaticMeshComponent.h
--- a/Engine/Source/Runtime/Engine/Classes/Components/StaticMeshComponent.h
+++ b/Engine/Source/Runtime/Engine/Classes/Components/StaticMeshComponent.h
@@ -303,6 +303,10 @@ public:
 
 	UPROPERTY(transient)
 	uint8 bDisplayVertexColors:1;
+
+	UPROPERTY(transient)
+	uint8 bDisplayPhysicalMaterialMasks : 1;
+
 #endif
 
 	/**
diff --git a/Engine/Source/Runtime/Engine/Classes/Engine/Engine.h b/Engine/Source/Runtime/Engine/Classes/Engine/Engine.h
--- a/Engine/Source/Runtime/Engine/Classes/Engine/Engine.h
+++ b/Engine/Source/Runtime/Engine/Classes/Engine/Engine.h
@@ -1027,6 +1027,14 @@ public:
 	UPROPERTY(globalconfig)
 	FSoftObjectPath ClothPaintMaterialWireframeName;
 
+	/** A material used to render physical material mask on mesh. */
+	UPROPERTY()
+	class UMaterial* PhysicalMaterialMaskMaterial;
+
+	/** A material used to render physical material mask on mesh. */
+	UPROPERTY(globalconfig)
+	FSoftObjectPath PhysicalMaterialMaskMaterialName;
+
 	/** A material used to render debug meshes. */
 	UPROPERTY()
 	class UMaterial* DebugEditorMaterial;
diff --git a/Engine/Source/Runtime/Engine/Private/Components/StaticMeshComponent.cpp b/Engine/Source/Runtime/Engine/Private/Components/StaticMeshComponent.cpp
--- a/Engine/Source/Runtime/Engine/Private/Components/StaticMeshComponent.cpp
+++ b/Engine/Source/Runtime/Engine/Private/Components/StaticMeshComponent.cpp
@@ -179,6 +179,7 @@ UStaticMeshComponent::UStaticMeshComponent(const FObjectInitializer& ObjectIniti
 	StaticMeshImportVersion = BeforeImportStaticMeshVersionWasAdded;
 	bCustomOverrideVertexColorPerLOD = false;
 	bDisplayVertexColors = false;
+	bDisplayPhysicalMaterialMasks = false;
 #endif
 }
 
@@ -912,7 +913,7 @@ bool UStaticMeshComponent::DoesSocketExist(FName InSocketName) const
 bool UStaticMeshComponent::ShouldRenderSelected() const
 {
 	const bool bShouldRenderSelected = UMeshComponent::ShouldRenderSelected();
-	return bShouldRenderSelected || bDisplayVertexColors;
+	return bShouldRenderSelected || bDisplayVertexColors || bDisplayPhysicalMaterialMasks;
 }
 #endif // WITH_EDITOR
 
diff --git a/Engine/Source/Runtime/Engine/Private/StaticMeshRender.cpp b/Engine/Source/Runtime/Engine/Private/StaticMeshRender.cpp
--- a/Engine/Source/Runtime/Engine/Private/StaticMeshRender.cpp
+++ b/Engine/Source/Runtime/Engine/Private/StaticMeshRender.cpp
@@ -27,6 +27,7 @@
 #include "EngineUtils.h"
 #include "StaticMeshResources.h"
 #include "SpeedTreeWind.h"
+#include "PhysicalMaterials/PhysicalMaterialMask.h"
 
 #include "Engine/Engine.h"
 #include "Engine/LevelStreaming.h"
@@ -1321,8 +1322,42 @@ void FStaticMeshSceneProxy::GetDynamicMeshElements(const TArray<const FSceneView
 								
 									if (GetMeshElement(LODIndex, BatchIndex, SectionIndex, SDPG_World, bSectionIsSelected, true, MeshElement))
 									{
-	#if STATICMESH_ENABLE_DEBUG_RENDERING
-										if (bProxyIsSelected && EngineShowFlags.VertexColors && AllowDebugViewmodes())
+										bool bDebugMaterialRenderProxySet = false;
+#if STATICMESH_ENABLE_DEBUG_RENDERING
+
+	#if WITH_EDITOR								
+										if (bProxyIsSelected && EngineShowFlags.PhysicalMaterialMasks && AllowDebugViewmodes())
+										{
+											// Override the mesh's material with our material that draws the physical material masks
+											UMaterial* PhysMatMaskVisualizationMaterial = GEngine->PhysicalMaterialMaskMaterial;
+											check(PhysMatMaskVisualizationMaterial);
+											
+											FMaterialRenderProxy* PhysMatMaskVisualizationMaterialInstance = nullptr;
+
+											const FLODInfo::FSectionInfo& Section = LODs[LODIndex].Sections[SectionIndex];
+											
+											if (UMaterialInterface* SectionMaterial = Section.Material)
+											{
+												if (UPhysicalMaterialMask* PhysicalMaterialMask = SectionMaterial->GetPhysicalMaterialMask())
+												{
+													if (PhysicalMaterialMask->MaskTexture)
+													{
+														PhysMatMaskVisualizationMaterialInstance = new FColoredTexturedMaterialRenderProxy(
+															PhysMatMaskVisualizationMaterial->GetRenderProxy(),
+															FLinearColor::White, NAME_Color, PhysicalMaterialMask->MaskTexture, NAME_LinearColor);
+													}
+
+													Collector.RegisterOneFrameMaterialProxy(PhysMatMaskVisualizationMaterialInstance);
+													MeshElement.MaterialRenderProxy = PhysMatMaskVisualizationMaterialInstance;
+
+													bDebugMaterialRenderProxySet = true;
+												}
+											}
+										}
+
+	#endif // WITH_EDITOR
+
+										if (!bDebugMaterialRenderProxySet && bProxyIsSelected && EngineShowFlags.VertexColors && AllowDebugViewmodes())
 										{
 											// Override the mesh's material with our material that draws the vertex colors
 											UMaterial* VertexColorVisualizationMaterial = NULL;
@@ -1357,11 +1392,13 @@ void FStaticMeshSceneProxy::GetDynamicMeshElements(const TArray<const FSceneView
 
 											Collector.RegisterOneFrameMaterialProxy(VertexColorVisualizationMaterialInstance);
 											MeshElement.MaterialRenderProxy = VertexColorVisualizationMaterialInstance;
+
+											bDebugMaterialRenderProxySet = true;
 										}
-										else
-	#endif
+
+	#endif // STATICMESH_ENABLE_DEBUG_RENDERING
 	#if WITH_EDITOR
-										if (bSectionIsSelected)
+										if (!bDebugMaterialRenderProxySet && bSectionIsSelected)
 										{
 											// Override the mesh's material with our material that draws the collision color
 											MeshElement.MaterialRenderProxy = new FOverrideSelectionColorMaterialRenderProxy(
@@ -1632,6 +1669,7 @@ FPrimitiveViewRelevance FStaticMeshSceneProxy::GetViewRelevance(const FSceneView
 #endif
 #if WITH_EDITOR
 		(IsSelected() && View->Family->EngineShowFlags.VertexColors) ||
+		(IsSelected() && View->Family->EngineShowFlags.PhysicalMaterialMasks) ||
 #endif
 #if STATICMESH_ENABLE_DEBUG_RENDERING
 		bDrawMeshCollisionIfComplex ||
diff --git a/Engine/Source/Runtime/Engine/Private/UnrealEngine.cpp b/Engine/Source/Runtime/Engine/Private/UnrealEngine.cpp
--- a/Engine/Source/Runtime/Engine/Private/UnrealEngine.cpp
+++ b/Engine/Source/Runtime/Engine/Private/UnrealEngine.cpp
@@ -2519,6 +2519,7 @@ void UEngine::InitializeObjectReferences()
 		LoadSpecialMaterial(TEXT("ClothPaintMaterialName"), ClothPaintMaterialName.ToString(), ClothPaintMaterial, false);
 		LoadSpecialMaterial(TEXT("ClothPaintMaterialWireframeName"), ClothPaintMaterialWireframeName.ToString(), ClothPaintMaterialWireframe, false);
 		LoadSpecialMaterial(TEXT("DebugEditorMaterialName"), DebugEditorMaterialName.ToString(), DebugEditorMaterial, false);
+		LoadSpecialMaterial(TEXT("PhysicalMaterialMaskMaterialName"), PhysicalMaterialMaskMaterialName.ToString(), PhysicalMaterialMaskMaterial, false);
 
 		ClothPaintMaterialInstance = UMaterialInstanceDynamic::Create(ClothPaintMaterial, nullptr);
 		ClothPaintMaterialWireframeInstance = UMaterialInstanceDynamic::Create(ClothPaintMaterialWireframe, nullptr);
diff --git a/Engine/Source/Runtime/Engine/Public/ShowFlags.h b/Engine/Source/Runtime/Engine/Public/ShowFlags.h
--- a/Engine/Source/Runtime/Engine/Public/ShowFlags.h
+++ b/Engine/Source/Runtime/Engine/Public/ShowFlags.h
@@ -322,6 +322,7 @@ private:
 		SetTestImage(false);
 		SetVisualizeDOF(false);
 		SetVertexColors(false);
+		SetPhysicalMaterialMasks(false);
 		SetVisualizeMotionBlur(false);
 		SetSelectionOutline(false);
 		SetDebugAI(false);
diff --git a/Engine/Source/Runtime/Engine/Public/ShowFlagsValues.inl b/Engine/Source/Runtime/Engine/Public/ShowFlagsValues.inl
--- a/Engine/Source/Runtime/Engine/Public/ShowFlagsValues.inl
+++ b/Engine/Source/Runtime/Engine/Public/ShowFlagsValues.inl
@@ -86,6 +86,8 @@ SHOWFLAG_FIXED_IN_SHIPPING(0, TestImage, SFG_Developer, NSLOCTEXT("UnrealEd", "T
 SHOWFLAG_FIXED_IN_SHIPPING(0, VisualizeDOF, SFG_Visualize, NSLOCTEXT("UnrealEd", "VisualizeDOFSF", "Depth of Field Layers"))
 /** Show Vertex Colors */
 SHOWFLAG_FIXED_IN_SHIPPING(0, VertexColors, SFG_Advanced, NSLOCTEXT("UnrealEd", "VertexColorsSF", "Vertex Colors"))
+/** Show Physical Material Masks */
+SHOWFLAG_FIXED_IN_SHIPPING(0, PhysicalMaterialMasks, SFG_Advanced, NSLOCTEXT("UnrealEd", "PhysicalMaterialMasksSF", "Physical Material Masks"))
 /** Render Post process (screen space) distortion/refraction */
 SHOWFLAG_FIXED_IN_SHIPPING(1, Refraction, SFG_Developer, NSLOCTEXT("UnrealEd", "RefractionSF", "Refraction"))
 /** Usually set in game or when previewing Matinee but not in editor, used for motion blur or any kind of rendering features that rely on the former frame */
