Fixed CIS warning

diff --git a/Engine/Source/Editor/StaticMeshEditor/Private/StaticMeshEditorTools.cpp b/Engine/Source/Editor/StaticMeshEditor/Private/StaticMeshEditorTools.cpp
--- a/Engine/Source/Editor/StaticMeshEditor/Private/StaticMeshEditorTools.cpp
+++ b/Engine/Source/Editor/StaticMeshEditor/Private/StaticMeshEditorTools.cpp
@@ -2133,7 +2133,7 @@ void FMeshSectionSettingsLayout::OnSectionForceOpaqueFlagChanged( ECheckBoxState
 	FScopedTransaction Transaction(TransactionTest);
 
 	PRAGMA_DISABLE_DEPRECATION_WARNINGS
-		UProperty* Property = UStaticMesh::StaticClass()->FindPropertyByName(GET_MEMBER_NAME_STRING_CHECKED(UStaticMesh, SectionInfoMap));
+		FProperty* Property = UStaticMesh::StaticClass()->FindPropertyByName(GET_MEMBER_NAME_STRING_CHECKED(UStaticMesh, SectionInfoMap));
 	PRAGMA_ENABLE_DEPRECATION_WARNINGS
 
 		StaticMesh.PreEditChange(Property);
