StaticMeshEditor grid drawing now using EngineShowFlags instead of DrawHelper for grid drawing.
Toolbar grid toggle button now syncs with viewport grid states.  If any viewport is showing a grid, the toggle will be on.
Fixes introduced with when UE-86749 was fixed.

diff --git a/Engine/Source/Editor/StaticMeshEditor/Private/StaticMeshEditor.cpp b/Engine/Source/Editor/StaticMeshEditor/Private/StaticMeshEditor.cpp
--- a/Engine/Source/Editor/StaticMeshEditor/Private/StaticMeshEditor.cpp
+++ b/Engine/Source/Editor/StaticMeshEditor/Private/StaticMeshEditor.cpp
@@ -213,6 +213,19 @@ void FStaticMeshEditor::InitEditorForStaticMesh(UStaticMesh* ObjectToEdit)
 	StaticMeshDetailsView->RegisterInstancedCustomPropertyLayout( UStaticMesh::StaticClass(), LayoutCustomStaticMeshProperties );
 
  	SetEditorMesh(ObjectToEdit);
+
+	bool LocalDrawGrids = false;
+	TFunction<void(FName, TSharedPtr<IEditorViewportLayoutEntity>)> CheckShowGridFunc =
+		[this, &LocalDrawGrids](FName Name, TSharedPtr<IEditorViewportLayoutEntity> Entity)
+	{
+		TSharedRef<SStaticMeshEditorViewport> StaticMeshEditorViewport = StaticCastSharedRef<SStaticMeshEditorViewport>(Entity->AsWidget());
+		FStaticMeshEditorViewportClient& StaticMeshEditorViewportClient = StaticMeshEditorViewport->GetViewportClient();
+		LocalDrawGrids |= StaticMeshEditorViewportClient.IsSetShowGridChecked();
+	};
+
+	ViewportTabContent->PerformActionOnViewports(CheckShowGridFunc);
+
+	bDrawGrids = LocalDrawGrids;
 }
 
 void FStaticMeshEditor::InitStaticMeshEditor( const EToolkitMode::Type Mode, const TSharedPtr< class IToolkitHost >& InitToolkitHost, UStaticMesh* ObjectToEdit )
@@ -2506,7 +2519,7 @@ bool FStaticMeshEditor::IsShowVerticesChecked() const
 
 void FStaticMeshEditor::ToggleShowGrids()
 {
-	bDrawGrids = !bDrawGrids;
+	bDrawGrids = !IsShowGridsChecked();
 	TFunction<void(FName, TSharedPtr<IEditorViewportLayoutEntity>)> ToggleShowGridFunc =
 		[this](FName Name, TSharedPtr<IEditorViewportLayoutEntity> Entity)
 	{
@@ -2525,7 +2538,18 @@ void FStaticMeshEditor::ToggleShowGrids()
 
 bool FStaticMeshEditor::IsShowGridsChecked() const
 {
-	return bDrawGrids;
+	bool LocalDrawGrids = false;
+	TFunction<void(FName, TSharedPtr<IEditorViewportLayoutEntity>)> CheckShowGridFunc =
+		[this, &LocalDrawGrids](FName Name, TSharedPtr<IEditorViewportLayoutEntity> Entity)
+	{
+		TSharedRef<SStaticMeshEditorViewport> StaticMeshEditorViewport = StaticCastSharedRef<SStaticMeshEditorViewport>(Entity->AsWidget());
+		FStaticMeshEditorViewportClient& StaticMeshEditorViewportClient = StaticMeshEditorViewport->GetViewportClient();
+		LocalDrawGrids |= StaticMeshEditorViewportClient.IsSetShowGridChecked();
+	};
+
+	ViewportTabContent->PerformActionOnViewports(CheckShowGridFunc);
+
+	return LocalDrawGrids;
 }
 
 
diff --git a/Engine/Source/Editor/StaticMeshEditor/Private/StaticMeshEditorViewportClient.cpp b/Engine/Source/Editor/StaticMeshEditor/Private/StaticMeshEditorViewportClient.cpp
--- a/Engine/Source/Editor/StaticMeshEditor/Private/StaticMeshEditorViewportClient.cpp
+++ b/Engine/Source/Editor/StaticMeshEditor/Private/StaticMeshEditorViewportClient.cpp
@@ -57,7 +57,7 @@ FStaticMeshEditorViewportClient::FStaticMeshEditorViewportClient(TWeakPtr<IStati
 	DrawHelper.bDrawPivot = false;
 	DrawHelper.bDrawWorldBox = false;
 	DrawHelper.bDrawKillZ = false;
-	DrawHelper.bDrawGrid = true;
+	DrawHelper.bDrawGrid = false;
 	DrawHelper.GridColorAxis = FColor(160,160,160);
 	DrawHelper.GridColorMajor = FColor(144,144,144);
 	DrawHelper.GridColorMinor = FColor(128,128,128);
@@ -1494,10 +1494,10 @@ void FStaticMeshEditorViewportClient::SetShowPivots(bool bShowOn)
 
 void FStaticMeshEditorViewportClient::SetShowGrids(bool bShowOn)
 {
-	DrawHelper.bDrawGrid = bShowOn;
+ 	EngineShowFlags.Grid = bShowOn;
 	if (FEngineAnalytics::IsAvailable())
 	{
-		FEngineAnalytics::GetProvider().RecordEvent(TEXT("Editor.Usage.StaticMesh.Toolbar"), TEXT("bDrawGrids"), bDrawVertices ? TEXT("True") : TEXT("False"));
+		FEngineAnalytics::GetProvider().RecordEvent(TEXT("Editor.Usage.StaticMesh.Toolbar"), TEXT("EngineShowFlags.Grid"), EngineShowFlags.Grid ? TEXT("True") : TEXT("False"));
 	}
 	Invalidate();
 }
diff --git a/Engine/Source/Editor/UnrealEd/Private/EditorViewportClient.cpp b/Engine/Source/Editor/UnrealEd/Private/EditorViewportClient.cpp
--- a/Engine/Source/Editor/UnrealEd/Private/EditorViewportClient.cpp
+++ b/Engine/Source/Editor/UnrealEd/Private/EditorViewportClient.cpp
@@ -412,7 +412,7 @@ FEditorViewportClient::FEditorViewportClient(FEditorModeTools* InModeTools, FPre
 	DrawHelper.bDrawPivot = false;
 	DrawHelper.bDrawWorldBox = false;
 	DrawHelper.bDrawKillZ = false;
-	DrawHelper.bDrawGrid = true;
+	DrawHelper.bDrawGrid = false; // disable this since we rely on the show flags
 	DrawHelper.GridColorAxis = FColor(160, 160, 160);
 	DrawHelper.GridColorMajor = FColor(144, 144, 144);
 	DrawHelper.GridColorMinor = FColor(128, 128, 128);
@@ -5127,17 +5127,17 @@ bool FEditorViewportClient::IsCameraLocked() const
 
 void FEditorViewportClient::SetShowGrid()
 {
-	DrawHelper.bDrawGrid = !DrawHelper.bDrawGrid;
+	EngineShowFlags.SetGrid(!EngineShowFlags.Grid) ;
 	if (FEngineAnalytics::IsAvailable())
 	{
-		FEngineAnalytics::GetProvider().RecordEvent(TEXT("Editor.Usage.StaticMesh.Toolbar"), TEXT("bDrawGrid"), DrawHelper.bDrawGrid ? TEXT("True") : TEXT("False"));
+		FEngineAnalytics::GetProvider().RecordEvent(TEXT("Editor.Usage.StaticMesh.Toolbar"), TEXT("EngineShowFlags.Grid"), EngineShowFlags.Grid ? TEXT("True") : TEXT("False"));
 	}
 	Invalidate();
 }
 
 bool FEditorViewportClient::IsSetShowGridChecked() const
 {
-	return DrawHelper.bDrawGrid;
+	return EngineShowFlags.Grid;
 }
 
 void FEditorViewportClient::SetShowBounds(bool bShow)
