Copying over 11179103  so that Rset Camera toolbar button will work correctly

diff --git a/Engine/Source/Editor/StaticMeshEditor/Private/StaticMeshEditor.cpp b/Engine/Source/Editor/StaticMeshEditor/Private/StaticMeshEditor.cpp
--- a/Engine/Source/Editor/StaticMeshEditor/Private/StaticMeshEditor.cpp
+++ b/Engine/Source/Editor/StaticMeshEditor/Private/StaticMeshEditor.cpp
@@ -749,6 +749,12 @@ void FStaticMeshEditor::BindCommands()
 		Commands.CreateAutoConvexCollision,
 		FExecuteAction::CreateSP(this, &FStaticMeshEditor::OnConvexDecomposition));
 
+	// Viewport Camera
+	UICommandList->MapAction(
+		Commands.ResetCamera,
+		FExecuteAction::CreateSP(this, &FStaticMeshEditor::ResetCamera));
+
+
 
 }
 
@@ -2713,6 +2719,24 @@ bool FStaticMeshEditor::IsShowVertexColorsChecked() const
 	return bDrawVertexColors;
 }
 
+void FStaticMeshEditor::ResetCamera()
+{
+	TFunction<void(FName, TSharedPtr<IEditorViewportLayoutEntity>)> ResetCameraFunc =
+		[this](FName Name, TSharedPtr<IEditorViewportLayoutEntity> Entity)
+	{
+		TSharedRef<SStaticMeshEditorViewport> StaticMeshEditorViewport = StaticCastSharedRef<SStaticMeshEditorViewport>(Entity->AsWidget());
+		StaticMeshEditorViewport->GetViewportClient().FocusViewportOnBox(StaticMeshEditorViewport->GetStaticMeshComponent()->Bounds.GetBox());
+		StaticMeshEditorViewport->Invalidate();
+	};
+
+	ViewportTabContent->PerformActionOnViewports(ResetCameraFunc);
+
+// 	if (FEngineAnalytics::IsAvailable())
+// 	{
+// 		FEngineAnalytics::GetProvider().RecordEvent(TEXT("Editor.Usage.StaticMesh.Toolbar"), TEXT("ResetCamera"));
+// 	}
+}
+
 void FStaticMeshEditor::RemoveCurrentUVChannel()
 {
 	if (!StaticMesh)
diff --git a/Engine/Source/Editor/StaticMeshEditor/Private/StaticMeshEditor.h b/Engine/Source/Editor/StaticMeshEditor/Private/StaticMeshEditor.h
--- a/Engine/Source/Editor/StaticMeshEditor/Private/StaticMeshEditor.h
+++ b/Engine/Source/Editor/StaticMeshEditor/Private/StaticMeshEditor.h
@@ -291,6 +291,9 @@ private:
 	/** Callback for checking the vertex colors show flag. */
 	bool IsShowVertexColorsChecked() const;
 
+	/** Reset the viewport camera to look at the mesh. */
+	void ResetCamera();
+
 private:
 
 
