Fix for the "Additional Data" button in the Static Mesh editor, so that the changes are propagated to individual viewports.

diff --git a/Engine/Source/Editor/StaticMeshEditor/Private/StaticMeshEditor.cpp b/Engine/Source/Editor/StaticMeshEditor/Private/StaticMeshEditor.cpp
--- a/Engine/Source/Editor/StaticMeshEditor/Private/StaticMeshEditor.cpp
+++ b/Engine/Source/Editor/StaticMeshEditor/Private/StaticMeshEditor.cpp
@@ -754,8 +754,12 @@ void FStaticMeshEditor::BindCommands()
 		Commands.ResetCamera,
 		FExecuteAction::CreateSP(this, &FStaticMeshEditor::ResetCamera));
 
-
-
+	// Draw additional data
+	UICommandList->MapAction(
+		Commands.SetDrawAdditionalData,
+		FExecuteAction::CreateSP(this, &FStaticMeshEditor::ToggleDrawAdditionalData),
+		FCanExecuteAction(),
+		FIsActionChecked::CreateSP(this, &FStaticMeshEditor::IsDrawAdditionalDataChecked));
 }
 
 static TSharedRef< SWidget > GenerateCollisionMenuContent(TSharedPtr<const FUICommandList> InCommandList)
@@ -2737,6 +2741,33 @@ void FStaticMeshEditor::ResetCamera()
 // 	}
 }
 
+void FStaticMeshEditor::ToggleDrawAdditionalData()
+{
+	bDrawAdditionalData = !bDrawAdditionalData;
+	TFunction<void(FName, TSharedPtr<IEditorViewportLayoutEntity>)> ToggleDrawAdditionalDataFunc =
+		[this](FName Name, TSharedPtr<IEditorViewportLayoutEntity> Entity)
+	{
+		TSharedRef<SStaticMeshEditorViewport> StaticMeshEditorViewport = StaticCastSharedRef<SStaticMeshEditorViewport>(Entity->AsWidget());
+		FStaticMeshEditorViewportClient& StaticMeshEditorViewportClient = StaticMeshEditorViewport->GetViewportClient();
+		if (StaticMeshEditorViewportClient.IsDrawAdditionalDataChecked() != bDrawAdditionalData)
+		{
+			StaticMeshEditorViewportClient.ToggleDrawAdditionalData();
+		}
+	};
+
+	ViewportTabContent->PerformActionOnViewports(ToggleDrawAdditionalDataFunc);
+
+	if (FEngineAnalytics::IsAvailable())
+	{
+		FEngineAnalytics::GetProvider().RecordEvent(TEXT("Editor.Usage.StaticMesh.Toolbar"), TEXT("bDrawAdditionalData"), bDrawAdditionalData ? TEXT("True") : TEXT("False"));
+	}
+}
+
+bool FStaticMeshEditor::IsDrawAdditionalDataChecked() const
+{
+	return bDrawAdditionalData;
+}
+
 void FStaticMeshEditor::RemoveCurrentUVChannel()
 {
 	if (!StaticMesh)
diff --git a/Engine/Source/Editor/StaticMeshEditor/Private/StaticMeshEditor.h b/Engine/Source/Editor/StaticMeshEditor/Private/StaticMeshEditor.h
--- a/Engine/Source/Editor/StaticMeshEditor/Private/StaticMeshEditor.h
+++ b/Engine/Source/Editor/StaticMeshEditor/Private/StaticMeshEditor.h
@@ -59,6 +59,7 @@ public:
 		, bDrawSockets(false)
 		, bDrawWireframes(false)
 		, bDrawVertexColors(false)
+		, bDrawAdditionalData(true)
 	{}
 
 	~FStaticMeshEditor();
@@ -294,6 +295,12 @@ private:
 	/** Reset the viewport camera to look at the mesh. */
 	void ResetCamera();
 
+	/** Callback for toggling the draw additional data flag. */
+	void ToggleDrawAdditionalData();
+
+	/** Callback for checking the draw additional data flag. */
+	bool IsDrawAdditionalDataChecked() const;
+
 private:
 
 
@@ -555,4 +562,5 @@ private:
 	bool bDrawSockets;
 	bool bDrawWireframes;
 	bool bDrawVertexColors;
+	bool bDrawAdditionalData;
 };
