Static Mesh Asset Editor Style Updates

diff --git a/Engine/Source/Developer/MeshMergeUtilities/Private/MeshMergeEditorExtensions.cpp b/Engine/Source/Developer/MeshMergeUtilities/Private/MeshMergeEditorExtensions.cpp
deleted file mode 100644
diff --git a/Engine/Source/Developer/MeshMergeUtilities/Private/MeshMergeEditorExtensions.h b/Engine/Source/Developer/MeshMergeUtilities/Private/MeshMergeEditorExtensions.h
deleted file mode 100644
diff --git a/Engine/Source/Developer/MeshMergeUtilities/Private/MeshMergeModule.cpp b/Engine/Source/Developer/MeshMergeUtilities/Private/MeshMergeModule.cpp
--- a/Engine/Source/Developer/MeshMergeUtilities/Private/MeshMergeModule.cpp
+++ b/Engine/Source/Developer/MeshMergeUtilities/Private/MeshMergeModule.cpp
@@ -3,7 +3,6 @@
 #include "MeshMergeModule.h"
 #include "MeshMergeUtilities.h"
 #include "Modules/ModuleManager.h"
-#include "MeshMergeEditorExtensions.h"
 #include "ToolMenus.h"
 
 class FMeshMergeModule : public IMeshMergeModule
@@ -19,26 +18,15 @@ public:
 		return Utilities;
 	}
 
-	virtual void StartupModule() override
-	{
-		UToolMenus::RegisterStartupCallback(FSimpleMulticastDelegate::FDelegate::CreateRaw(this, &FMeshMergeModule::RegisterMenus));
-	}
 
 	virtual void ShutdownModule() override
 	{
-		UToolMenus::UnRegisterStartupCallback(this);
 		UToolMenus::UnregisterOwner(this);
 	}
 
 protected:
 	FMeshMergeUtilities Utilities;
 
-private:
-	void RegisterMenus()
-	{
-		FToolMenuOwnerScoped OwnerScoped(this);
-		FMeshMergeEditorExtensions::RegisterMenus();
-	}
 };
 
 
diff --git a/Engine/Source/Developer/MeshMergeUtilities/Private/MeshMergeUtilities.cpp b/Engine/Source/Developer/MeshMergeUtilities/Private/MeshMergeUtilities.cpp
--- a/Engine/Source/Developer/MeshMergeUtilities/Private/MeshMergeUtilities.cpp
+++ b/Engine/Source/Developer/MeshMergeUtilities/Private/MeshMergeUtilities.cpp
@@ -50,7 +50,6 @@
 #include "StaticMeshComponentAdapter.h"
 #include "SkeletalMeshAdapter.h"
 #include "StaticMeshAdapter.h"
-#include "MeshMergeEditorExtensions.h"
 
 #include "MeshMergeDataTracker.h"
 
@@ -87,15 +86,11 @@ DEFINE_LOG_CATEGORY(LogMeshMerging);
 FMeshMergeUtilities::FMeshMergeUtilities()
 {
 	Processor = new FProxyGenerationProcessor(this);
-
-	// Add callback for registering editor extensions with Skeletal/Static mesh editor
-	ModuleLoadedDelegateHandle = FModuleManager::Get().OnModulesChanged().AddStatic(&FMeshMergeEditorExtensions::OnModulesChanged);
 }
 
 FMeshMergeUtilities::~FMeshMergeUtilities()
 {
 	FModuleManager::Get().OnModulesChanged().Remove(ModuleLoadedDelegateHandle);
-	FMeshMergeEditorExtensions::RemoveExtenders();
 }
 
 void FMeshMergeUtilities::BakeMaterialsForComponent(TArray<TWeakObjectPtr<UObject>>& OptionObjects, IMaterialBakingAdapter* Adapter) const
diff --git a/Engine/Source/Editor/StaticMeshEditor/Private/SStaticMeshEditorViewport.cpp b/Engine/Source/Editor/StaticMeshEditor/Private/SStaticMeshEditorViewport.cpp
--- a/Engine/Source/Editor/StaticMeshEditor/Private/SStaticMeshEditorViewport.cpp
+++ b/Engine/Source/Editor/StaticMeshEditor/Private/SStaticMeshEditorViewport.cpp
@@ -79,9 +79,14 @@ void SStaticMeshEditorViewport::PopulateViewportOverlays(TSharedRef<SOverlay> Ov
 	Overlay->AddSlot()
 		.VAlign(VAlign_Top)
 		.HAlign(HAlign_Left)
-		.Padding(FMargin(10.0f, 40.0f, 10.0f, 10.0f))
+		.Padding(FMargin(6.0f, 36.0f, 6.0f, 6.0f))
 		[
-			SAssignNew(OverlayTextVerticalBox, SVerticalBox)
+			SNew(SBorder)
+			.BorderImage( FAppStyle::Get().GetBrush( "FloatingBorder" ) )
+			.Padding(4.f)
+			[
+				SAssignNew(OverlayTextVerticalBox, SVerticalBox)
+			]
 		];
 
 	// this widget will display the current viewed feature level
diff --git a/Engine/Source/Editor/StaticMeshEditor/Private/StaticMeshEditor.cpp b/Engine/Source/Editor/StaticMeshEditor/Private/StaticMeshEditor.cpp
--- a/Engine/Source/Editor/StaticMeshEditor/Private/StaticMeshEditor.cpp
+++ b/Engine/Source/Editor/StaticMeshEditor/Private/StaticMeshEditor.cpp
@@ -17,6 +17,7 @@
 #include "StaticMeshEditorModule.h"
 #include "AI/Navigation/NavCollisionBase.h"
 #include "ScopedTransaction.h"
+#include "ToolMenus.h"
 
 #include "SStaticMeshEditorViewport.h"
 #include "PropertyEditorModule.h"
@@ -31,6 +32,7 @@
 #include "EditorViewportCommands.h"
 #include "Editor/UnrealEd/Private/ConvexDecompTool.h"
 
+#include "MeshMergeModule.h"
 #include "Runtime/Analytics/Analytics/Public/Interfaces/IAnalyticsProvider.h"
 #include "EngineAnalytics.h"
 #include "Widgets/Docking/SDockTab.h"
@@ -82,32 +84,32 @@ void FStaticMeshEditor::RegisterTabSpawners(const TSharedRef<class FTabManager>&
 	InTabManager->RegisterTabSpawner( ViewportTabId, FOnSpawnTab::CreateSP(this, &FStaticMeshEditor::SpawnTab_Viewport) )
 		.SetDisplayName( LOCTEXT("ViewportTab", "Viewport") )
 		.SetGroup(WorkspaceMenuCategoryRef)
-		.SetIcon(FSlateIcon(FEditorStyle::GetStyleSetName(), "LevelEditor.Tabs.Viewports"));
+		.SetIcon(FSlateIcon(FAppStyle::GetAppStyleSetName(), "LevelEditor.Tabs.Viewports"));
 
 	InTabManager->RegisterTabSpawner( PropertiesTabId, FOnSpawnTab::CreateSP(this, &FStaticMeshEditor::SpawnTab_Properties) )
 		.SetDisplayName( LOCTEXT("PropertiesTab", "Details") )
 		.SetGroup(WorkspaceMenuCategoryRef)
-		.SetIcon(FSlateIcon(FEditorStyle::GetStyleSetName(), "LevelEditor.Tabs.Details"));
+		.SetIcon(FSlateIcon(FAppStyle::GetAppStyleSetName(), "LevelEditor.Tabs.Details"));
 
 	InTabManager->RegisterTabSpawner( SocketManagerTabId, FOnSpawnTab::CreateSP(this, &FStaticMeshEditor::SpawnTab_SocketManager) )
 		.SetDisplayName( LOCTEXT("SocketManagerTab", "Socket Manager") )
 		.SetGroup(WorkspaceMenuCategoryRef)
-		.SetIcon(FSlateIcon(FEditorStyle::GetStyleSetName(), "StaticMeshEditor.Tabs.SocketManager"));
+		.SetIcon(FSlateIcon(FAppStyle::GetAppStyleSetName(), "StaticMeshEditor.Tabs.SocketManager"));
 
 	InTabManager->RegisterTabSpawner( CollisionTabId, FOnSpawnTab::CreateSP(this, &FStaticMeshEditor::SpawnTab_Collision) )
 		.SetDisplayName( LOCTEXT("CollisionTab", "Convex Decomposition") )
 		.SetGroup(WorkspaceMenuCategoryRef)
-		.SetIcon(FSlateIcon(FEditorStyle::GetStyleSetName(), "StaticMeshEditor.Tabs.ConvexDecomposition"));
+		.SetIcon(FSlateIcon(FAppStyle::GetAppStyleSetName(), "StaticMeshEditor.Tabs.ConvexDecomposition"));
 
 	InTabManager->RegisterTabSpawner( PreviewSceneSettingsTabId, FOnSpawnTab::CreateSP(this, &FStaticMeshEditor::SpawnTab_PreviewSceneSettings) )
 		.SetDisplayName( LOCTEXT("PreviewSceneTab", "Preview Scene Settings") )
 		.SetGroup(WorkspaceMenuCategoryRef)
-		.SetIcon(FSlateIcon(FEditorStyle::GetStyleSetName(), "LevelEditor.Tabs.Details"));
+		.SetIcon(FSlateIcon(FAppStyle::GetAppStyleSetName(), "LevelEditor.Tabs.Details"));
 
 	FTabSpawnerEntry& MenuEntry = InTabManager->RegisterTabSpawner( SecondaryToolbarTabId, FOnSpawnTab::CreateSP(this, &FStaticMeshEditor::SpawnTab_SecondaryToolbar) )
 		.SetDisplayName( LOCTEXT("ToolbarTab", "Secondary Toolbar") )
 		.SetGroup(WorkspaceMenuCategoryRef)
-		.SetIcon(FSlateIcon(FEditorStyle::GetStyleSetName(), "Toolbar.Icon"));
+		.SetIcon(FSlateIcon(FAppStyle::GetAppStyleSetName(), "Toolbar.Icon"));
 
 	// Hide the menu item by default. It will be enabled only if the secondary toolbar is populated with extensions
 	SecondaryToolbarEntry = &MenuEntry;
@@ -212,24 +214,22 @@ void FStaticMeshEditor::InitStaticMeshEditor( const EToolkitMode::Type Mode, con
 	//Let additional extensions dock themselves to this TabStack of tools
 	OnStaticMeshEditorDockingExtentionTabs().Broadcast(ExtentionTabStack);
 
-	const TSharedRef<FTabManager::FLayout> StandaloneDefaultLayout = FTabManager::NewLayout( "Standalone_StaticMeshEditor_Layout_v4.2" )
+	const TSharedRef<FTabManager::FLayout> StandaloneDefaultLayout = FTabManager::NewLayout( "Standalone_StaticMeshEditor_Layout_v5" )
 	->AddArea
 	(
 		FTabManager::NewPrimaryArea() ->SetOrientation(Orient_Vertical)
 		->Split
-		(
-			FTabManager::NewStack()
-			->SetSizeCoefficient(0.1f)
-			->SetHideTabWell( true )
-			->AddTab(GetToolbarTabId(), ETabState::OpenedTab)
-			// Don't want the secondary toolbar tab to be opened if there's nothing in it
-			->AddTab(SecondaryToolbarTabId, ETabState::ClosedTab)
-		)
-		->Split
 		(
 			FTabManager::NewSplitter() ->SetOrientation(Orient_Horizontal)
 			->SetSizeCoefficient(0.9f)
 			->Split
+			(
+				FTabManager::NewStack()
+				->SetSizeCoefficient(0.15f)
+				->AddTab(SocketManagerTabId, ETabState::SidebarTab, ESidebarLocation::Left, 0.15)
+				->SetHideTabWell( true )
+			)
+			->Split
 			(
 				FTabManager::NewStack()
 				->SetSizeCoefficient(0.6f)
@@ -244,7 +244,6 @@ void FStaticMeshEditor::InitStaticMeshEditor( const EToolkitMode::Type Mode, con
 				(
 					FTabManager::NewStack()
 					->SetSizeCoefficient(0.7f)
-					->AddTab(PreviewSceneSettingsTabId, ETabState::OpenedTab)
 					->AddTab(PropertiesTabId, ETabState::OpenedTab)
 				)
 				->Split
@@ -369,29 +368,13 @@ void FStaticMeshEditor::ExtendMenu()
 			InMenuBuilder.EndSection();
 		}
 
-		static void FillMeshMenu( FMenuBuilder& InMenuBuilder )
+		static void FillMeshMenu( FToolMenuSection& InMenuSection)
 		{
-			// @todo mainframe: These menus, and indeed all menus like them, should be updated with extension points, plus expose public module
-			// access to extending the menus.  They may also need to extend the command list, or be able to PUSH a command list of their own.
-			// If we decide to only allow PUSHING, then nothing else should be needed (happens by extender automatically).  But if we want to
-			// augment the asset editor's existing command list, then we need to think about how to expose support for that.
-
-			InMenuBuilder.BeginSection("MeshFindSource");
-			{
-				InMenuBuilder.AddMenuEntry(FStaticMeshEditorCommands::Get().FindSource);
-			}
-			InMenuBuilder.EndSection();
-
-			InMenuBuilder.BeginSection("MeshChange");
+			static auto* CVar = IConsoleManager::Get().FindTConsoleVariableDataInt(TEXT("r.StaticMesh.EnableSaveGeneratedLODsInPackage"));
+			if (CVar && CVar->GetValueOnGameThread() != 0)
 			{
-				InMenuBuilder.AddMenuEntry(FStaticMeshEditorCommands::Get().ChangeMesh);
-				static auto* CVar = IConsoleManager::Get().FindTConsoleVariableDataInt(TEXT("r.StaticMesh.EnableSaveGeneratedLODsInPackage"));
-				if (CVar && CVar->GetValueOnGameThread() != 0)
-				{
-					InMenuBuilder.AddMenuEntry(FStaticMeshEditorCommands::Get().SaveGeneratedLODs);
-				}
+				InMenuSection.AddMenuEntry(FStaticMeshEditorCommands::Get().SaveGeneratedLODs);
 			}
-			InMenuBuilder.EndSection();
 		}
 
 		static void FillCollisionMenu( FMenuBuilder& InMenuBuilder )
@@ -425,16 +408,28 @@ void FStaticMeshEditor::ExtendMenu()
 				InMenuBuilder.AddMenuEntry(FStaticMeshEditorCommands::Get().CopyCollisionFromSelectedMesh);
 			}
 			InMenuBuilder.EndSection();
+
+			InMenuBuilder.BeginSection("MeshFindSource");
+			{
+				InMenuBuilder.AddMenuEntry(FStaticMeshEditorCommands::Get().FindSource);
+			}
+			InMenuBuilder.EndSection();
+
+			InMenuBuilder.BeginSection("MeshChange");
+			{
+				// InMenuBuilder.AddMenuEntry(FStaticMeshEditorCommands::Get().ChangeMesh);
+				static auto* CVar = IConsoleManager::Get().FindTConsoleVariableDataInt(TEXT("r.StaticMesh.EnableSaveGeneratedLODsInPackage"));
+				if (CVar && CVar->GetValueOnGameThread() != 0)
+				{
+					InMenuBuilder.AddMenuEntry(FStaticMeshEditorCommands::Get().SaveGeneratedLODs);
+				}
+			}
+			InMenuBuilder.EndSection();
+
 		}
 
-		static void GenerateMeshAndCollisionMenuBars( FMenuBarBuilder& InMenuBarBuilder)
+		static void GenerateCollisionMenuBar( FMenuBarBuilder& InMenuBarBuilder)
 		{
-			InMenuBarBuilder.AddPullDownMenu(
-				LOCTEXT("StaticMeshEditorMeshMenu", "Mesh"),
-				LOCTEXT("StaticMeshEditorMeshMenu_ToolTip", "Opens a menu with commands for altering this mesh"),
-				FNewMenuDelegate::CreateStatic(&Local::FillMeshMenu),
-				"Mesh");
-
 			InMenuBarBuilder.AddPullDownMenu(
 				LOCTEXT("StaticMeshEditorCollisionMenu", "Collision"),
 				LOCTEXT("StaticMeshEditorCollisionMenu_ToolTip", "Opens a menu with commands for editing this mesh's collision"),
@@ -451,17 +446,26 @@ void FStaticMeshEditor::ExtendMenu()
 		GetToolkitCommands(),
 		FMenuExtensionDelegate::CreateStatic( &Local::FillEditMenu ) );
 
+
 	MenuExtender->AddMenuBarExtension(
 		"Asset",
 		EExtensionHook::After,
 		GetToolkitCommands(),
-		FMenuBarExtensionDelegate::CreateStatic( &Local::GenerateMeshAndCollisionMenuBars )
+		FMenuBarExtensionDelegate::CreateStatic( &Local::GenerateCollisionMenuBar )
 		);
 
 	AddMenuExtender(MenuExtender);
 
 	IStaticMeshEditorModule* StaticMeshEditorModule = &FModuleManager::LoadModuleChecked<IStaticMeshEditorModule>( "StaticMeshEditor" );
 	AddMenuExtender(StaticMeshEditorModule->GetMenuExtensibilityManager()->GetAllExtenders(GetToolkitCommands(), GetEditingObjects()));
+
+	UToolMenu* AssetMenu = UToolMenus::Get()->ExtendMenu("MainFrame.MainMenu.Asset");
+	FToolMenuSection& AssetSection = AssetMenu->AddSection("StaticMeshAssets");
+	AssetSection.AddMenuEntry(FStaticMeshEditorCommands::Get().FindSource);
+	AssetSection.AddMenuEntry(FStaticMeshEditorCommands::Get().SetDrawAdditionalData );
+	AssetSection.AddMenuEntry(FStaticMeshEditorCommands::Get().BakeMaterials);
+	AssetSection.AddDynamicEntry("SaveGeneratedLODs", FNewToolMenuSectionDelegate::CreateStatic(&Local::FillMeshMenu));
+
 }
 
 void FStaticMeshEditor::AddReferencedObjects( FReferenceCollector& Collector )
@@ -508,6 +512,13 @@ TSharedRef<SDockTab> FStaticMeshEditor::SpawnTab_SocketManager( const FSpawnTabA
 {
 	check( Args.GetTabId() == SocketManagerTabId );
 
+	if (!SocketManager)
+	{
+		FSimpleDelegate OnSocketSelectionChanged = FSimpleDelegate::CreateSP( SharedThis(this), &FStaticMeshEditor::OnSocketSelectionChanged );
+
+		SocketManager = ISocketManager::CreateSocketManager( SharedThis(this) , OnSocketSelectionChanged );
+}
+
 	return SNew(SDockTab)
 		.Label( LOCTEXT("StaticMeshSocketManager_TabTitle", "Socket Manager") )
 		[
@@ -519,6 +530,12 @@ TSharedRef<SDockTab> FStaticMeshEditor::SpawnTab_Collision( const FSpawnTabArgs&
 {
 	check( Args.GetTabId() == CollisionTabId );
 
+	if (!ConvexDecomposition)
+	{
+		SAssignNew( ConvexDecomposition, SConvexDecomposition )
+			.StaticMeshEditorPtr(SharedThis(this));
+	}
+
 	return SNew(SDockTab)
 		.Label( LOCTEXT("StaticMeshConvexDecomp_TabTitle", "Convex Decomposition") )
 		[
@@ -758,16 +775,61 @@ void FStaticMeshEditor::BindCommands()
 		FExecuteAction::CreateSP(this, &FStaticMeshEditor::ToggleDrawAdditionalData),
 		FCanExecuteAction(),
 		FIsActionChecked::CreateSP(this, &FStaticMeshEditor::IsDrawAdditionalDataChecked));
+
+	// Bake Materials
+	UICommandList->MapAction(
+		Commands.BakeMaterials,
+		FExecuteAction::CreateSP(this, &FStaticMeshEditor::BakeMaterials),
+		FCanExecuteAction());
 }
 
 static TSharedRef< SWidget > GenerateCollisionMenuContent(TSharedPtr<const FUICommandList> InCommandList)
 {
 	FMenuBuilder MenuBuilder(true, InCommandList);
+	MenuBuilder.BeginSection("CollisionEditCollision");
+	{
+		MenuBuilder.AddMenuEntry(FStaticMeshEditorCommands::Get().CreateSphereCollision);
+		MenuBuilder.AddMenuEntry(FStaticMeshEditorCommands::Get().CreateSphylCollision);
+		MenuBuilder.AddMenuEntry(FStaticMeshEditorCommands::Get().CreateBoxCollision);
+		MenuBuilder.AddMenuEntry(FStaticMeshEditorCommands::Get().CreateDOP10X);
+		MenuBuilder.AddMenuEntry(FStaticMeshEditorCommands::Get().CreateDOP10Y);
+		MenuBuilder.AddMenuEntry(FStaticMeshEditorCommands::Get().CreateDOP10Z);
+		MenuBuilder.AddMenuEntry(FStaticMeshEditorCommands::Get().CreateDOP18);
+		MenuBuilder.AddMenuEntry(FStaticMeshEditorCommands::Get().CreateDOP26);
+		MenuBuilder.AddMenuEntry(FStaticMeshEditorCommands::Get().ConvertBoxesToConvex);
+		MenuBuilder.AddMenuEntry(FStaticMeshEditorCommands::Get().RemoveCollision);
+		
+		MenuBuilder.AddMenuEntry(FGenericCommands::Get().Delete, "DeleteCollision", LOCTEXT("DeleteCollision", "Delete Selected Collision"), LOCTEXT("DeleteCollisionToolTip", "Deletes the selected Collision from the mesh."));
+		MenuBuilder.AddMenuEntry(FGenericCommands::Get().Duplicate, "DuplicateCollision", LOCTEXT("DuplicateCollision", "Duplicate Selected Collision"), LOCTEXT("DuplicateCollisionToolTip", "Duplicates the selected Collision."));
+	}
+	MenuBuilder.EndSection();
+
+	MenuBuilder.BeginSection("CollisionAutoConvexCollision");
+	{
+		MenuBuilder.AddMenuEntry(FStaticMeshEditorCommands::Get().CreateAutoConvexCollision);
+	}
+	MenuBuilder.EndSection();
 
-	MenuBuilder.BeginSection("ShowCollision", LOCTEXT("ShowCollision", "Show Collision"));
+	MenuBuilder.BeginSection("CollisionCopy");
 	{
-		MenuBuilder.AddMenuEntry(FStaticMeshEditorCommands::Get().ToggleShowSimpleCollisions);
-		MenuBuilder.AddMenuEntry(FStaticMeshEditorCommands::Get().ToggleShowComplexCollisions);
+		MenuBuilder.AddMenuEntry(FStaticMeshEditorCommands::Get().CopyCollisionFromSelectedMesh);
+	}
+	MenuBuilder.EndSection();
+
+	MenuBuilder.BeginSection("MeshFindSource");
+	{
+		MenuBuilder.AddMenuEntry(FStaticMeshEditorCommands::Get().FindSource);
+	}
+	MenuBuilder.EndSection();
+
+	MenuBuilder.BeginSection("MeshChange");
+	{
+		// MenuBuilder.AddMenuEntry(FStaticMeshEditorCommands::Get().ChangeMesh);
+		static auto* CVar = IConsoleManager::Get().FindTConsoleVariableDataInt(TEXT("r.StaticMesh.EnableSaveGeneratedLODsInPackage"));
+		if (CVar && CVar->GetValueOnGameThread() != 0)
+		{
+			MenuBuilder.AddMenuEntry(FStaticMeshEditorCommands::Get().SaveGeneratedLODs);
+		}
 	}
 	MenuBuilder.EndSection();
 
@@ -794,12 +856,6 @@ void FStaticMeshEditor::ExtendToolBar()
 				return MenuBuilder.MakeWidget();
 			};
 
-			ToolbarBuilder.BeginSection("Realtime");
-			{
-				ToolbarBuilder.AddToolBarButton(FEditorViewportCommands::Get().ToggleRealTime);
-			}
-			ToolbarBuilder.EndSection();
-
 			ToolbarBuilder.BeginSection("Mesh");
 			{
 				ToolbarBuilder.AddToolBarButton(FUIAction(FExecuteAction::CreateSP(ThisEditor, &FStaticMeshEditor::HandleReimportMesh)),
@@ -818,12 +874,6 @@ void FStaticMeshEditor::ExtendToolBar()
 
 			ToolbarBuilder.BeginSection("Command");
 			{
-				ToolbarBuilder.AddToolBarButton(FStaticMeshEditorCommands::Get().ToggleShowSockets);
-				ToolbarBuilder.AddToolBarButton(FStaticMeshEditorCommands::Get().ToggleShowWireframes);
-				ToolbarBuilder.AddToolBarButton(FStaticMeshEditorCommands::Get().ToggleShowVertexColors);
-				ToolbarBuilder.AddToolBarButton(FStaticMeshEditorCommands::Get().ToggleShowGrids);
-				ToolbarBuilder.AddToolBarButton(FStaticMeshEditorCommands::Get().ToggleShowBounds);
-
 				TSharedPtr<const FUICommandList> CommandList = ToolbarBuilder.GetTopCommandList();
 
 				ToolbarBuilder.AddComboButton(
@@ -834,12 +884,6 @@ void FStaticMeshEditor::ExtendToolBar()
 					FSlateIcon(FEditorStyle::GetStyleSetName(), "StaticMeshEditor.SetShowCollision")
 				);
 
-				ToolbarBuilder.AddToolBarButton(FStaticMeshEditorCommands::Get().ToggleShowPivots);
-				ToolbarBuilder.AddToolBarButton(FStaticMeshEditorCommands::Get().ToggleShowNormals);
-				ToolbarBuilder.AddToolBarButton(FStaticMeshEditorCommands::Get().ToggleShowTangents);
-				ToolbarBuilder.AddToolBarButton(FStaticMeshEditorCommands::Get().ToggleShowBinormals);
-				ToolbarBuilder.AddToolBarButton(FStaticMeshEditorCommands::Get().ToggleShowVertices);
-
 				FOnGetContent OnGetUVMenuContent = FOnGetContent::CreateRaw(ThisEditor, &FStaticMeshEditor::GenerateUVChannelComboList);
 
 				ToolbarBuilder.AddComboButton(
@@ -851,13 +895,6 @@ void FStaticMeshEditor::ExtendToolBar()
 			}
 
 			ToolbarBuilder.EndSection();
-			ToolbarBuilder.BeginSection("Camera");
-			{
-				ToolbarBuilder.AddToolBarButton(FStaticMeshEditorCommands::Get().ResetCamera);
-			}
-			ToolbarBuilder.EndSection();
-
-			ToolbarBuilder.AddToolBarButton(FStaticMeshEditorCommands::Get().SetDrawAdditionalData);
 		}
 	};
 
@@ -888,13 +925,6 @@ void FStaticMeshEditor::ExtendToolBar()
 
 void FStaticMeshEditor::BuildSubTools()
 {
-	FSimpleDelegate OnSocketSelectionChanged = FSimpleDelegate::CreateSP( SharedThis(this), &FStaticMeshEditor::OnSocketSelectionChanged );
-
-	SocketManager = ISocketManager::CreateSocketManager( SharedThis(this) , OnSocketSelectionChanged );
-
-	SAssignNew( ConvexDecomposition, SConvexDecomposition )
-		.StaticMeshEditorPtr(SharedThis(this));
-
 	FAdvancedPreviewSceneModule& AdvancedPreviewSceneModule = FModuleManager::LoadModuleChecked<FAdvancedPreviewSceneModule>("AdvancedPreviewScene");
 
 	TArray<FAdvancedPreviewSceneModule::FDetailDelegates> Delegates;
@@ -2788,6 +2818,15 @@ bool FStaticMeshEditor::IsDrawAdditionalDataChecked() const
 	return bDrawAdditionalData;
 }
 
+void FStaticMeshEditor::BakeMaterials()
+{
+	if (StaticMesh != nullptr)
+	{	
+		const IMeshMergeModule& Module = FModuleManager::Get().LoadModuleChecked<IMeshMergeModule>("MeshMergeUtilities");
+		Module.GetUtilities().BakeMaterialsForMesh(StaticMesh);
+	}
+}
+
 void FStaticMeshEditor::RemoveCurrentUVChannel()
 {
 	if (!StaticMesh)
diff --git a/Engine/Source/Editor/StaticMeshEditor/Private/StaticMeshEditor.h b/Engine/Source/Editor/StaticMeshEditor/Private/StaticMeshEditor.h
--- a/Engine/Source/Editor/StaticMeshEditor/Private/StaticMeshEditor.h
+++ b/Engine/Source/Editor/StaticMeshEditor/Private/StaticMeshEditor.h
@@ -303,6 +303,9 @@ private:
 	/** Callback for checking the draw additional data flag. */
 	bool IsDrawAdditionalDataChecked() const;
 
+	/** Bake out Materials for give LOD(s). **/
+	void BakeMaterials();
+
 private:
 
 
@@ -429,6 +432,7 @@ private:
 	void OnConvexDecomposition();
 
 	//~ Begin FAssetEditorToolkit Interface.
+	virtual bool UsesCustomToolbarPlacement() const { return false;  }
 	virtual bool OnRequestClose() override;
 	//~ End FAssetEditorToolkit Interface.
 
diff --git a/Engine/Source/Editor/StaticMeshEditor/Private/StaticMeshEditorActions.cpp b/Engine/Source/Editor/StaticMeshEditor/Private/StaticMeshEditorActions.cpp
--- a/Engine/Source/Editor/StaticMeshEditor/Private/StaticMeshEditorActions.cpp
+++ b/Engine/Source/Editor/StaticMeshEditor/Private/StaticMeshEditorActions.cpp
@@ -15,7 +15,7 @@ void FStaticMeshEditorCommands::RegisterCommands()
 	UI_COMMAND( SetShowSimpleCollision, "Simple Collision", "Toggles display of the simplified collision mesh of the static mesh, if present.", EUserInterfaceActionType::ToggleButton, FInputChord());
 	UI_COMMAND( SetShowComplexCollision, "Complex Collision", "Toggles display of the complex collision for this static mesh.", EUserInterfaceActionType::ToggleButton, FInputChord());
 	UI_COMMAND( ResetCamera, "Reset Camera", "Resets the camera to focus on the static mesh", EUserInterfaceActionType::Button, FInputChord() );
-	UI_COMMAND( SetShowSockets, "Sockets", "Displays the static mesh sockets.", EUserInterfaceActionType::ToggleButton, FInputChord() );
+	UI_COMMAND( SetShowSockets, "Sockets", "Displays the static mesh sockets.", EUserInterfaceActionType::ToggleButton, FInputChord(EModifierKey::Alt, EKeys::S) );
 	UI_COMMAND( SetDrawAdditionalData, "Additional Data", "Draw additional user data associated with asset.", EUserInterfaceActionType::ToggleButton, FInputChord() );
 
 	UI_COMMAND( ReimportMesh, "Reimport Base Mesh", "Reimport the base mesh.", EUserInterfaceActionType::Button, FInputChord() );
@@ -36,11 +36,11 @@ void FStaticMeshEditorCommands::RegisterCommands()
 	UI_COMMAND( ToggleShowWireframes, "Wireframe", "Toggles the viewmode of the Preview Pane between a lit view and a wireframe view.", EUserInterfaceActionType::ToggleButton, FInputChord());
 	UI_COMMAND( ToggleShowVertexColors, "Vert Colors", "Toggles vertex colors.", EUserInterfaceActionType::ToggleButton, FInputChord());
 
-	UI_COMMAND( SetShowNormals, "Normals", "Toggles display of vertex normals in the Preview Pane.", EUserInterfaceActionType::ToggleButton, FInputChord() );
-	UI_COMMAND( SetShowTangents, "Tangents", "Toggles display of vertex tangents in the Preview Pane.", EUserInterfaceActionType::ToggleButton, FInputChord() );
-	UI_COMMAND( SetShowBinormals, "Binormals", "Toggles display of vertex binormals (orthogonal vector to normal and tangent) in the Preview Pane.", EUserInterfaceActionType::ToggleButton, FInputChord() );
-	UI_COMMAND( SetShowVertices, "Vertices", "Toggles display of vertices in the Preview Pane.", EUserInterfaceActionType::ToggleButton, FInputChord() ); 
-	UI_COMMAND( SetShowPivot, "Show Pivot", "Display the pivot location of the static mesh.", EUserInterfaceActionType::ToggleButton, FInputChord() );
+	UI_COMMAND( SetShowNormals, "Normals", "Toggles display of vertex normals in the Preview Pane.", EUserInterfaceActionType::ToggleButton, FInputChord(EModifierKey::Alt, EKeys::N) );
+	UI_COMMAND( SetShowTangents, "Tangents", "Toggles display of vertex tangents in the Preview Pane.", EUserInterfaceActionType::ToggleButton, FInputChord(EModifierKey::Alt, EKeys::T) );
+	UI_COMMAND( SetShowBinormals, "Binormals", "Toggles display of vertex binormals (orthogonal vector to normal and tangent) in the Preview Pane.", EUserInterfaceActionType::ToggleButton, FInputChord(EModifierKey::Alt, EKeys::B) );
+	UI_COMMAND( SetShowVertices, "Vertices", "Toggles display of vertices in the Preview Pane.", EUserInterfaceActionType::ToggleButton, FInputChord(EModifierKey::Alt, EKeys::V) ); 
+	UI_COMMAND( SetShowPivot, "Show Pivot", "Display the pivot location of the static mesh.", EUserInterfaceActionType::ToggleButton, FInputChord(EModifierKey::Alt, EKeys::P) );
 
 	UI_COMMAND( CreateDOP10X, "Add 10DOP-X Simplified Collision", "Generates a new axis-aligned box collision mesh with the 4 X-axis aligned edges beveled (10 total sides) encompassing the static mesh.", EUserInterfaceActionType::Button, FInputChord() );
 	UI_COMMAND( CreateDOP10Y, "Add 10DOP-Y Simplified Collision", "Generates a new axis-aligned box collision mesh with the 4 Y-axis aligned edges beveled (10 total sides) encompassing the static mesh.", EUserInterfaceActionType::Button, FInputChord() );
@@ -62,6 +62,7 @@ void FStaticMeshEditorCommands::RegisterCommands()
 	UI_COMMAND( ChangeMesh, "Change Mesh", "Changes the static mesh asset loaded in the Static Mesh Editor to the asset currently selected in the Content Browser.", EUserInterfaceActionType::Button, FInputChord() );
 
 	UI_COMMAND( SaveGeneratedLODs, "Save Generated LODs", "Saves the generated LODs in the mesh package.", EUserInterfaceActionType::Button, FInputChord());
+	UI_COMMAND( BakeMaterials, "Bake Out Materials", "Bake out Materials for given LOD(s).", EUserInterfaceActionType::Button, FInputChord());
 }
 
 #undef LOCTEXT_NAMESPACE
diff --git a/Engine/Source/Editor/StaticMeshEditor/Public/StaticMeshEditorActions.h b/Engine/Source/Editor/StaticMeshEditor/Public/StaticMeshEditorActions.h
--- a/Engine/Source/Editor/StaticMeshEditor/Public/StaticMeshEditorActions.h
+++ b/Engine/Source/Editor/StaticMeshEditor/Public/StaticMeshEditorActions.h
@@ -88,6 +88,8 @@ public:
 
 	TSharedPtr< FUICommandInfo > ChangeMesh;
 
+	TSharedPtr< FUICommandInfo > BakeMaterials;
+
 	TSharedPtr< FUICommandInfo > SaveGeneratedLODs;
 
 	/**
diff --git a/Engine/Source/Editor/StaticMeshEditor/StaticMeshEditor.Build.cs b/Engine/Source/Editor/StaticMeshEditor/StaticMeshEditor.Build.cs
--- a/Engine/Source/Editor/StaticMeshEditor/StaticMeshEditor.Build.cs
+++ b/Engine/Source/Editor/StaticMeshEditor/StaticMeshEditor.Build.cs
@@ -43,7 +43,8 @@ public class StaticMeshEditor : ModuleRules
                 "DesktopWidgets",
 				"EditorSubsystem",
 				"MeshDescription",
-				"StaticMeshDescription"
+				"StaticMeshDescription",
+				"ToolMenus"
 			}
 		);
 
diff --git a/Engine/Source/Editor/UnrealEd/Private/SSocketManager.cpp b/Engine/Source/Editor/UnrealEd/Private/SSocketManager.cpp
--- a/Engine/Source/Editor/UnrealEd/Private/SSocketManager.cpp
+++ b/Engine/Source/Editor/UnrealEd/Private/SSocketManager.cpp
@@ -45,8 +45,9 @@ public:
 	FOnRenameRequested OnRenameRequested;
 };
 
-class SSocketDisplayItem : public STableRow< TSharedPtr<FString> >
+class SSocketDisplayItem : public SMultiColumnTableRow< TSharedPtr<SocketListItem> >
 {
+
 public:
 	
 	SLATE_BEGIN_ARGS( SSocketDisplayItem )
@@ -69,31 +70,58 @@ public:
 		SocketItem = InArgs._SocketItem;
 		SocketManagerPtr = InArgs._SocketManagerPtr;
 
-		TSharedPtr< SInlineEditableTextBlock > InlineWidget;
+		ImportedSocketBrush = FAppStyle::Get().GetBrush("Icons.Import");
+		NotImportedSocketBrush = FAppStyle::Get().GetBrush("NoBrush");
+
+		auto Args = FSuperRowType::FArguments();
+			// .Style(&FEditorStyle::Get().GetWidgetStyle<FTableRowStyle>("SceneOutliner.TableViewRow"));
+
+		SMultiColumnTableRow< TSharedPtr<SocketListItem> >::Construct(Args, InOwnerTableView);
+	}
+
+	virtual TSharedRef<SWidget> GenerateWidgetForColumn( const FName& ColumnName ) override
+	{
+		// todo make these static id names 
+		if( ColumnName == TEXT("Socket"))
+		{
+			TSharedPtr< SInlineEditableTextBlock > InlineWidget;
 
-		this->ChildSlot
-		.Padding( 0.0f, 3.0f, 6.0f, 3.0f )
-		.VAlign(VAlign_Center)
-		[
 			SAssignNew( InlineWidget, SInlineEditableTextBlock )
-				.Text( this, &SSocketDisplayItem::GetSocketName )
-				.OnVerifyTextChanged( this, &SSocketDisplayItem::OnVerifySocketNameChanged )
-				.OnTextCommitted( this, &SSocketDisplayItem::OnCommitSocketName )
-				.IsSelected( this, &STableRow< TSharedPtr<FString> >::IsSelectedExclusively )
-		];
+			.Text( this, &SSocketDisplayItem::GetSocketName )
+			.OnVerifyTextChanged( this, &SSocketDisplayItem::OnVerifySocketNameChanged )
+			.OnTextCommitted( this, &SSocketDisplayItem::OnCommitSocketName )
+			.IsSelected( this, &SMultiColumnTableRow< TSharedPtr<SocketListItem> >::IsSelectedExclusively );
 
-		TSharedPtr<SocketListItem> SocketItemPinned = SocketItem.Pin();
-		if (SocketItemPinned.IsValid())
+			TSharedPtr<SocketListItem> SocketItemPinned = SocketItem.Pin();
+			if (SocketItemPinned.IsValid())
+			{
+				SocketItemPinned->OnRenameRequested.BindSP(InlineWidget.Get(), &SInlineEditableTextBlock::EnterEditingMode);
+			}
+
+			return SNew(SBox)
+			.HAlign(HAlign_Left)
+			.VAlign(VAlign_Center)
+			.Padding(FMargin(20.f, 4.f))
+			[
+				InlineWidget.ToSharedRef()
+				];
+		}
+		else
 		{
-			SocketItemPinned->OnRenameRequested.BindSP(InlineWidget.Get(), &SInlineEditableTextBlock::EnterEditingMode);
+			return SNew(SBox)
+			.HAlign(HAlign_Right)
+			.VAlign(VAlign_Center)
+			[
+				SNew(SImage)
+				.ToolTipText(this, &SSocketDisplayItem::GetImportToolTip)
+				.ColorAndOpacity(FSlateColor::UseForeground())
+				.Image(this, &SSocketDisplayItem::GetImportBrush)
+				];
 		}
 
-		STableRow< TSharedPtr<FString> >::ConstructInternal(
-			STableRow::FArguments()
-				.ShowSelection(true),
-			InOwnerTableView
-		);	
+		return SNullWidget::NullWidget;
 	}
+
 private:
 	/** Returns the socket name */
 	FText GetSocketName() const
@@ -102,6 +130,25 @@ private:
 		return SocketItemPinned.IsValid() ? FText::FromName(SocketItemPinned->Socket->SocketName) : FText();
 	}
 
+	const FSlateBrush* GetImportBrush() const
+	{
+
+		TSharedPtr<SocketListItem> SocketItemPinned = SocketItem.Pin();
+		return SocketItemPinned.IsValid() && SocketItemPinned->Socket->bSocketCreatedAtImport ? ImportedSocketBrush : NotImportedSocketBrush;
+
+	}
+
+	FText GetImportToolTip() const
+	{
+		TSharedPtr<SocketListItem> SocketItemPinned = SocketItem.Pin();
+		if (SocketItemPinned.IsValid() && SocketItemPinned->Socket->bSocketCreatedAtImport)
+		{
+			return LOCTEXT("ImportedSockedTooltip", "Socket was imported from the source mesh.");
+		}
+
+		return FText::GetEmpty();
+	}
+
 	bool OnVerifySocketNameChanged( const FText& InNewText, FText& OutErrorMessage )
 	{
 		bool bVerifyName = true;
@@ -161,6 +208,9 @@ private:
 
 	/** Pointer back to the socket manager */
 	TWeakPtr< SSocketManager > SocketManagerPtr; 
+
+	const FSlateBrush* ImportedSocketBrush;
+	const FSlateBrush* NotImportedSocketBrush;
 };
 
 TSharedPtr<ISocketManager> ISocketManager::CreateSocketManager(TSharedPtr<class IStaticMeshEditor> InStaticMeshEditor, FSimpleDelegate InOnSocketSelectionChanged )
@@ -213,7 +263,7 @@ void SSocketManager::Construct(const FArguments& InArgs)
 		.FillHeight(1.0f)
 		[
 			SNew(SSplitter)
-			.Orientation(Orient_Horizontal)
+			.Orientation(Orient_Vertical)
 
 			+ SSplitter::Slot()
 			.Value(.3f)
@@ -225,14 +275,43 @@ void SSocketManager::Construct(const FArguments& InArgs)
 
 					+ SVerticalBox::Slot()
 					.AutoHeight()
-					.Padding(0, 0, 0, 4)
 					[
-						SNew(SButton)
-						.ButtonStyle(FEditorStyle::Get(), "FlatButton.Success")
-						.ForegroundColor(FLinearColor::White)
-						.Text(LOCTEXT("CreateSocket", "Create Socket"))
-						.OnClicked(this, &SSocketManager::CreateSocket_Execute)
-						.HAlign(HAlign_Center)
+						SNew(SHorizontalBox)
+						+SHorizontalBox::Slot()
+						[
+							SNew(SBorder)
+							.BorderImage(FAppStyle::Get().GetBrush("Brushes.Background"))
+							.HAlign(HAlign_Left)
+							.Padding(FMargin(12.f, 6.f))
+							[
+								SNew(STextBlock)
+								.TextStyle(FAppStyle::Get(), "ButtonText")
+								.Text(LOCTEXT("Sockets", "Sockets"))
+								.TransformPolicy(ETextTransformPolicy::ToUpper)	
+							]
+						]
+
+						+SHorizontalBox::Slot()
+						.AutoWidth()
+						.VAlign(VAlign_Center)
+						[
+							SNew(SButton)
+							.ButtonStyle(FAppStyle::Get(), "SimpleButton")
+							.ToolTipText(LOCTEXT("CreateSocket", "Create Socket"))
+							.OnClicked(this, &SSocketManager::CreateSocket_Execute)
+							.VAlign(VAlign_Center)
+							[
+								SNew(SImage)
+								.ColorAndOpacity(FSlateColor::UseForeground())
+								.Image(FAppStyle::Get().GetBrush("Icons.Plus"))
+								]
+							]
+						]
+
+					+ SVerticalBox::Slot()
+					.AutoHeight()
+					[
+						SNew(SSeparator)
 					]
 
 					+ SVerticalBox::Slot()
@@ -251,27 +330,29 @@ void SSocketManager::Construct(const FArguments& InArgs)
 						.OnSelectionChanged(this, &SSocketManager::SocketSelectionChanged_Execute)
 
 						// Allow for some spacing between items with a larger item height.
-						.ItemHeight(20.0f)
+						.ItemHeight(24.0f)
 
 						.OnContextMenuOpening(this, &SSocketManager::OnContextMenuOpening)
 						.OnItemScrolledIntoView(this, &SSocketManager::OnItemScrolledIntoView)
 
 						.HeaderRow
 						(
-						SNew(SHeaderRow)
-						.Visibility(EVisibility::Collapsed)
-						+ SHeaderRow::Column(TEXT("Socket"))
-						)
-					]
+							SNew(SHeaderRow)
+							.Visibility(EVisibility::Collapsed)
 
-					+ SVerticalBox::Slot()
-					.AutoHeight()
-					[
-						SNew(SSeparator)
+							+ SHeaderRow::Column(TEXT("Socket"))
+							.HAlignCell(HAlign_Fill)
+							.VAlignCell(VAlign_Center)
+
+							+ SHeaderRow::Column(TEXT("Imported"))
+							.HAlignCell(HAlign_Right)
+							.VAlignCell(VAlign_Center)
+						)
 					]
 
 					+ SVerticalBox::Slot()
 					.AutoHeight()
+					.Padding(FMargin(12.f, 6.f))
 					[
 						SNew(STextBlock)
 						.Text(this, &SSocketManager::GetSocketHeaderText)
diff --git a/Engine/Source/Runtime/Engine/Classes/Engine/StaticMeshSocket.h b/Engine/Source/Runtime/Engine/Classes/Engine/StaticMeshSocket.h
--- a/Engine/Source/Runtime/Engine/Classes/Engine/StaticMeshSocket.h
+++ b/Engine/Source/Runtime/Engine/Classes/Engine/StaticMeshSocket.h
@@ -34,15 +34,15 @@ class UStaticMeshSocket : public UObject
 	UPROPERTY(EditAnywhere, BlueprintReadWrite, Category=StaticMeshSocket)
 	FVector RelativeScale;
 
-	UPROPERTY(EditAnywhere, BlueprintReadWrite, Category=StaticMeshSocket)
+	UPROPERTY(EditAnywhere, AdvancedDisplay, BlueprintReadWrite, Category=StaticMeshSocket)
 	FString Tag;
 
 #if WITH_EDITORONLY_DATA
-	UPROPERTY(EditAnywhere, Category=StaticMeshSocket)
+	UPROPERTY(EditAnywhere, AdvancedDisplay, Category=StaticMeshSocket)
 	class UStaticMesh* PreviewStaticMesh;
 
 	/* Weather the socket was imported with the asset or create in the editor. Importer will remove/modify only imported socket and will not touch any editor created socket. */
-	UPROPERTY(VisibleAnywhere, AdvancedDisplay, Category = StaticMeshSocket)
+	UPROPERTY()
 	bool bSocketCreatedAtImport;
 #endif // WITH_EDITORONLY_DATA
 
