Reconfigured the default position of the Static Mesh Socket Manager Panel

diff --git a/Engine/Source/Editor/StaticMeshEditor/Private/StaticMeshEditor.cpp b/Engine/Source/Editor/StaticMeshEditor/Private/StaticMeshEditor.cpp
--- a/Engine/Source/Editor/StaticMeshEditor/Private/StaticMeshEditor.cpp
+++ b/Engine/Source/Editor/StaticMeshEditor/Private/StaticMeshEditor.cpp
@@ -226,30 +226,24 @@ void FStaticMeshEditor::InitStaticMeshEditor( const EToolkitMode::Type Mode, con
 		->Split
 		(
 			FTabManager::NewSplitter() ->SetOrientation(Orient_Horizontal)
-			->SetSizeCoefficient(0.9f)
 			->Split
 			(
 				FTabManager::NewStack()
-				->SetSizeCoefficient(0.15f)
-				->AddTab(SocketManagerTabId, ETabState::SidebarTab, ESidebarLocation::Left, 0.15)
-				->SetHideTabWell( true )
-			)
-			->Split
-			(
-				FTabManager::NewStack()
-				->SetSizeCoefficient(0.6f)
+				->SetSizeCoefficient(0.7f)
 				->AddTab(ViewportTabId, ETabState::OpenedTab)
 				->SetHideTabWell( true )
 			)
 			->Split
 			(
 				FTabManager::NewSplitter() ->SetOrientation(Orient_Vertical)
-				->SetSizeCoefficient(0.2f)
+				->SetSizeCoefficient(0.25f)
 				->Split
 				(
 					FTabManager::NewStack()
 					->SetSizeCoefficient(0.7f)
 					->AddTab(PropertiesTabId, ETabState::OpenedTab)
+					->AddTab(SocketManagerTabId, ETabState::OpenedTab)
+					->SetForegroundTab(PropertiesTabId) 
 				)
 				->Split
 				(
