Adds back missing Show Vertex Colors Option

diff --git a/Engine/Source/Editor/StaticMeshEditor/Private/SStaticMeshEditorViewportToolBar.cpp b/Engine/Source/Editor/StaticMeshEditor/Private/SStaticMeshEditorViewportToolBar.cpp
--- a/Engine/Source/Editor/StaticMeshEditor/Private/SStaticMeshEditorViewportToolBar.cpp
+++ b/Engine/Source/Editor/StaticMeshEditor/Private/SStaticMeshEditorViewportToolBar.cpp
@@ -46,6 +46,7 @@ TSharedRef<SWidget> SStaticMeshEditorViewportToolbar::GenerateShowMenu() const
 		ShowMenuBuilder.AddMenuEntry(Commands.SetShowSockets);
 		ShowMenuBuilder.AddMenuEntry(Commands.SetShowPivot);
 		ShowMenuBuilder.AddMenuEntry(Commands.SetShowVertices);
+		ShowMenuBuilder.AddMenuEntry(Commands.SetShowVertexColor);
 
 		ShowMenuBuilder.AddMenuSeparator();
 
