Added function to EditorStaticMeshLibrary for setting material slot assignment for LODs
(Moved CL# 14951629 to the Static Mesh Editor Subsystem)

diff --git a/Engine/Source/Editor/StaticMeshEditor/Private/StaticMeshEditorSubsystem.cpp b/Engine/Source/Editor/StaticMeshEditor/Private/StaticMeshEditorSubsystem.cpp
--- a/Engine/Source/Editor/StaticMeshEditor/Private/StaticMeshEditorSubsystem.cpp
+++ b/Engine/Source/Editor/StaticMeshEditor/Private/StaticMeshEditorSubsystem.cpp
@@ -1432,6 +1432,50 @@ void UStaticMeshEditorSubsystem::EnableSectionCastShadow(UStaticMesh* StaticMesh
 	StaticMesh->PostEditChange();
 }
 
+void UStaticMeshEditorSubsystem::SetLODMaterialSlot(UStaticMesh* StaticMesh, int32 MaterialSlotIndex, int32 LODIndex, int32 SectionIndex)
+{
+	TGuardValue<bool> UnattendedScriptGuard(GIsRunningUnattendedScript, true);
+
+	if (!EditorScriptingHelpers::CheckIfInEditorAndPIE())
+	{
+		return;
+	}
+
+	if (StaticMesh == nullptr)
+	{
+		UE_LOG(LogStaticMeshEditorSubsystem, Error, TEXT("SetLODMaterialSlot: The StaticMesh is null."));
+		return;
+	}
+
+	if (LODIndex >= StaticMesh->GetNumLODs())
+	{
+		UE_LOG(LogStaticMeshEditorSubsystem, Error, TEXT("SetLODMaterialSlot: Invalid LOD index %d (of %d)."), LODIndex, StaticMesh->GetNumLODs());
+		return;
+	}
+
+	if (SectionIndex >= StaticMesh->GetNumSections(LODIndex))
+	{
+		UE_LOG(LogStaticMeshEditorSubsystem, Error, TEXT("SetLODMaterialSlot: Invalid section index %d (of %d)."), SectionIndex, StaticMesh->GetNumSections(LODIndex));
+		return;
+	}
+
+	if (MaterialSlotIndex >= StaticMesh->GetStaticMaterials().Num())
+	{
+		UE_LOG(LogStaticMeshEditorSubsystem, Error, TEXT("SetLODMaterialSlot: Invalid slot index %d (of %d)."), MaterialSlotIndex, StaticMesh->GetStaticMaterials().Num());
+		return;
+	}
+
+	StaticMesh->Modify();
+
+	FMeshSectionInfo SectionInfo = StaticMesh->GetSectionInfoMap().Get(LODIndex, SectionIndex);
+
+	SectionInfo.MaterialIndex = MaterialSlotIndex;
+
+	StaticMesh->GetSectionInfoMap().Set(LODIndex, SectionIndex, SectionInfo);
+
+	StaticMesh->PostEditChange();
+}
+
 bool UStaticMeshEditorSubsystem::HasVertexColors(UStaticMesh* StaticMesh)
 {
 	TGuardValue<bool> UnattendedScriptGuard(GIsRunningUnattendedScript, true);
diff --git a/Engine/Source/Editor/StaticMeshEditor/Public/StaticMeshEditorSubsystem.h b/Engine/Source/Editor/StaticMeshEditor/Public/StaticMeshEditorSubsystem.h
--- a/Engine/Source/Editor/StaticMeshEditor/Public/StaticMeshEditorSubsystem.h
+++ b/Engine/Source/Editor/StaticMeshEditor/Public/StaticMeshEditorSubsystem.h
@@ -294,6 +294,17 @@ public:
 	UFUNCTION(BlueprintCallable, Category = "Static Mesh Utilities")
 	void EnableSectionCastShadow(UStaticMesh* StaticMesh, bool bCastShadow, int32 LODIndex, int32 SectionIndex);
 
+	/**
+	* Sets the material slot for a specific LOD.
+	* @param	StaticMesh			Static mesh to Enables/disables shadow casting from.
+	* @param	MaterialSlotIndex	Index of the material slot to use.
+	* @param	LODIndex			Index of the StaticMesh LOD.
+	* @param	SectionIndex		Index of the StaticMesh Section.
+	*/
+	UFUNCTION(BlueprintCallable, Category = "Static Mesh Utilities")
+	void SetLODMaterialSlot(UStaticMesh* StaticMesh, int32 MaterialSlotIndex, int32 LODIndex, int32 SectionIndex);
+
+
 	/** Check whether a static mesh has vertex colors */
 	UFUNCTION(BlueprintPure, Category = "Static Mesh Utilities")
 	bool HasVertexColors(UStaticMesh* StaticMesh);
