StaticMeshEditor: gracefully handle the situation where simple collision on an SM is modified externally, as there is no change notification to hook into. SelectedPrims becomes invalid and this check() will be immediately hit, so we can handle it here (obviously a hack).
#rnx

diff --git a/Engine/Source/Editor/StaticMeshEditor/Private/StaticMeshEditor.cpp b/Engine/Source/Editor/StaticMeshEditor/Private/StaticMeshEditor.cpp
--- a/Engine/Source/Editor/StaticMeshEditor/Private/StaticMeshEditor.cpp
+++ b/Engine/Source/Editor/StaticMeshEditor/Private/StaticMeshEditor.cpp
@@ -1261,7 +1261,15 @@ bool FStaticMeshEditor::GetLastSelectedPrimTransform(FTransform& OutTransform) c
 
 		const FPrimData& PrimData = SelectedPrims.Last();
 
-		check(IsPrimValid(PrimData));
+		// The SME is not notified of external changes to Simple Collision of the active object, and the UBodySetup
+		// does not have any kind of change notification to hook into to do this. So, the SelectedPrims can become
+		// invalid if an external change is made. In that case we will just return false here and hope that the 
+		// FStaticMeshEditorViewportClient caller will error-handle correctly.
+		if (IsPrimValid(PrimData) == false)
+		{
+			return false;
+		}
+
 		switch (PrimData.PrimType)
 		{
 		case EAggCollisionShape::Sphere:
diff --git a/Engine/Source/Editor/StaticMeshEditor/Private/StaticMeshEditorViewportClient.cpp b/Engine/Source/Editor/StaticMeshEditor/Private/StaticMeshEditorViewportClient.cpp
--- a/Engine/Source/Editor/StaticMeshEditor/Private/StaticMeshEditorViewportClient.cpp
+++ b/Engine/Source/Editor/StaticMeshEditor/Private/StaticMeshEditorViewportClient.cpp
@@ -422,8 +422,11 @@ FVector FStaticMeshEditorViewportClient::GetWidgetLocation() const
 	{
 		return PrimTransform.GetLocation();
 	}
-
-	return FVector::ZeroVector;
+	else
+	{
+		StaticMeshEditorPtr.Pin()->ClearSelectedPrims();
+		return FVector::ZeroVector;
+	}
 }
 
 FMatrix FStaticMeshEditorViewportClient::GetWidgetCoordSystem() const 
@@ -447,8 +450,11 @@ FMatrix FStaticMeshEditorViewportClient::GetWidgetCoordSystem() const
 	{
 		return FRotationMatrix(PrimTransform.Rotator());
 	}
-
-	return FMatrix::Identity;
+	else
+	{
+		StaticMeshEditorPtr.Pin()->ClearSelectedPrims();
+		return FMatrix::Identity;
+	}
 }
 
 ECoordSystem FStaticMeshEditorViewportClient::GetWidgetCoordSystemSpace() const
