Static Mesh Nanite Options

diff --git a/Engine/Source/Editor/StaticMeshEditor/Private/StaticMeshEditorTools.cpp b/Engine/Source/Editor/StaticMeshEditor/Private/StaticMeshEditorTools.cpp
--- a/Engine/Source/Editor/StaticMeshEditor/Private/StaticMeshEditorTools.cpp
+++ b/Engine/Source/Editor/StaticMeshEditor/Private/StaticMeshEditorTools.cpp
@@ -3394,6 +3394,7 @@ void FLevelOfDetailSettingsLayout::AddLODLevelCategories( IDetailLayoutBuilder&
 			FString CategoryName = FString(TEXT("StaticMeshMaterials"));
 
 			IDetailCategoryBuilder& MaterialsCategory = DetailBuilder.EditCategory(*CategoryName, LOCTEXT("StaticMeshMaterialsLabel", "Material Slots"), ECategoryPriority::Important);
+			MaterialsCategory.SetSortOrder(0);
 			MaterialsLayoutWidget = MakeShareable(new FMeshMaterialsLayout(StaticMeshEditor));
 			TArray<FAssetData> AssetDataArray;
 			AssetDataArray.Add(FAssetData(StaticMesh, false));
@@ -4612,7 +4613,9 @@ void FNaniteSettingsLayout::AddToDetailsPanel(IDetailLayoutBuilder& DetailBuilde
 
 	IDetailCategoryBuilder& NaniteSettingsCategory =
 		DetailBuilder.EditCategory("NaniteSettings", LOCTEXT("NaniteSettingsCategory", "Nanite Settings"));
+	NaniteSettingsCategory.SetSortOrder(10);
 
+	TSharedPtr<SCheckBox> NaniteEnabledCheck;
 	{
 		NaniteSettingsCategory.AddCustomRow( LOCTEXT("Enabled", "Enabled") )
 		.NameContent()
@@ -4623,7 +4626,7 @@ void FNaniteSettingsLayout::AddToDetailsPanel(IDetailLayoutBuilder& DetailBuilde
 		]
 		.ValueContent()
 		[
-			SNew(SCheckBox)
+			SAssignNew(NaniteEnabledCheck, SCheckBox)
 			.IsChecked(this, &FNaniteSettingsLayout::IsEnabledChecked)
 			.OnCheckStateChanged(this, &FNaniteSettingsLayout::OnEnabledChanged)
 		];
@@ -4631,11 +4634,14 @@ void FNaniteSettingsLayout::AddToDetailsPanel(IDetailLayoutBuilder& DetailBuilde
 
 	{
 		NaniteSettingsCategory.AddCustomRow( LOCTEXT("ProxyTrianglePercent", "Proxy Triangle Percent") )
+
+		.IsEnabled(TAttribute<bool>::Create(TAttribute<bool>::FGetter::CreateLambda([NaniteEnabledCheck]() -> bool {return NaniteEnabledCheck->IsChecked(); } )))
 		.NameContent()
 		[
 			SNew(STextBlock)
 			.Font(IDetailLayoutBuilder::GetDetailFont())
 			.Text(LOCTEXT("ProxyTrianglePercent", "Proxy Triangle Percent"))
+			.ToolTipText(LOCTEXT("ProxyTrianglePercentTooltip", "Percentage of triangles to reduce down to for generating a coarse proxy\nmesh that will be used anywhere the full detail Nanite data can't,\nincluding platforms that don't support Nanite rendering."))
 		]
 		.ValueContent()
 		[
