Made NaniteSettings EditAnywhere on StaticMeshes.
Updated FStaticMeshDetails to prevent NaniteSettings showing twice

diff --git a/Engine/Source/Editor/StaticMeshEditor/Private/StaticMeshEditorTools.cpp b/Engine/Source/Editor/StaticMeshEditor/Private/StaticMeshEditorTools.cpp
--- a/Engine/Source/Editor/StaticMeshEditor/Private/StaticMeshEditorTools.cpp
+++ b/Engine/Source/Editor/StaticMeshEditor/Private/StaticMeshEditorTools.cpp
@@ -130,6 +130,10 @@ void FStaticMeshDetails::CustomizeDetails( class IDetailLayoutBuilder& DetailBui
 	LevelOfDetailSettings = MakeShareable( new FLevelOfDetailSettingsLayout( StaticMeshEditor ) );
 	LevelOfDetailSettings->AddToDetailsPanel( DetailBuilder );
 
+	// Hide the existing NaniteSettings property so we can use the customization instead
+	TSharedRef<IPropertyHandle> NaniteSettingsProperty = DetailBuilder.GetProperty(GET_MEMBER_NAME_CHECKED(UStaticMesh, NaniteSettings));
+	NaniteSettingsProperty->MarkHiddenByCustomization();
+
 	NaniteSettings = MakeShareable(new FNaniteSettingsLayout(StaticMeshEditor));
 	NaniteSettings->AddToDetailsPanel(DetailBuilder);
 
diff --git a/Engine/Source/Runtime/Engine/Classes/Engine/StaticMesh.h b/Engine/Source/Runtime/Engine/Classes/Engine/StaticMesh.h
--- a/Engine/Source/Runtime/Engine/Classes/Engine/StaticMesh.h
+++ b/Engine/Source/Runtime/Engine/Classes/Engine/StaticMesh.h
@@ -725,7 +725,7 @@ public:
 	TArray<TObjectPtr<UMaterialInterface>> Materials_DEPRECATED;
 
 	/** Settings related to building Nanite data. */
-	UPROPERTY()
+	UPROPERTY(EditAnywhere, Category=NaniteSettings)
 	FMeshNaniteSettings NaniteSettings;
 
 #endif // #if WITH_EDITORONLY_DATA
