Added a view menu toggle to the static mesh editor to show or hide the nanite proxy mesh. The option only appears if nanite is enabled on a mesh.

Added overlay text for nanite enabled and a warning when the proxy mesh is being displayed

Fixed overlay text creating new widgets every frame

diff --git a/Engine/Source/Editor/EditorStyle/Private/StarshipStyle.cpp b/Engine/Source/Editor/EditorStyle/Private/StarshipStyle.cpp
--- a/Engine/Source/Editor/EditorStyle/Private/StarshipStyle.cpp
+++ b/Engine/Source/Editor/EditorStyle/Private/StarshipStyle.cpp
@@ -334,7 +334,7 @@ void FStarshipEditorStyle::FStyle::SetupGeneralStyles()
 			.SetShadowColorAndOpacity(FLinearColor(0.0f, 0.0f, 0.0f)));
 
 		Set("TextBlock.ShadowedTextWarning", FTextBlockStyle(NormalText)
-			.SetColorAndOpacity(FLinearColor(1.0f, 0.0f, 0.0f))
+			.SetColorAndOpacity(FStyleColors::Warning)
 			.SetShadowOffset(FVector2D(1.0f, 1.0f))
 			.SetShadowColorAndOpacity(FLinearColor(0.0f, 0.0f, 0.0f)));
 
diff --git a/Engine/Source/Editor/StaticMeshEditor/Private/SStaticMeshEditorViewport.cpp b/Engine/Source/Editor/StaticMeshEditor/Private/SStaticMeshEditorViewport.cpp
--- a/Engine/Source/Editor/StaticMeshEditor/Private/SStaticMeshEditorViewport.cpp
+++ b/Engine/Source/Editor/StaticMeshEditor/Private/SStaticMeshEditorViewport.cpp
@@ -20,6 +20,7 @@
 #include "Engine/StaticMeshSocket.h"
 #include "Editor/UnrealEd/Public/SEditorViewportToolBarMenu.h"
 #include "Editor.h"
+#include "Widgets/Text/SRichTextBlock.h"
 
 #define HITPROXY_SOCKET	1
 
@@ -85,7 +86,7 @@ void SStaticMeshEditorViewport::PopulateViewportOverlays(TSharedRef<SOverlay> Ov
 			.BorderImage( FAppStyle::Get().GetBrush( "FloatingBorder" ) )
 			.Padding(4.f)
 			[
-				SAssignNew(OverlayTextVerticalBox, SVerticalBox)
+				SAssignNew(OverlayText, SRichTextBlock)
 			]
 		];
 
@@ -116,19 +117,26 @@ SStaticMeshEditorViewport::~SStaticMeshEditorViewport()
 	}
 }
 
-void SStaticMeshEditorViewport::PopulateOverlayText(const TArray<FOverlayTextItem>& TextItems)
+void SStaticMeshEditorViewport::PopulateOverlayText(const TArrayView<FOverlayTextItem> TextItems)
 {
-	OverlayTextVerticalBox->ClearChildren();
+	FTextBuilder FinalText;
+
+	static FText WarningTextStyle = FText::FromString(TEXT("<TextBlock.ShadowedTextWarning>{0}</>"));
+	static FText NormalTextStyle = FText::FromString(TEXT("<TextBlock.ShadowedText>{0}</>"));
 
 	for (const auto& TextItem : TextItems)
 	{
-		OverlayTextVerticalBox->AddSlot()
-		[
-			SNew(STextBlock)
-			.Text(TextItem.Text)
-			.TextStyle(FEditorStyle::Get(), TextItem.Style)
-		];
+		if (!TextItem.bIsCustomFormat)
+		{
+			FinalText.AppendLineFormat(TextItem.bIsWarning ? WarningTextStyle : NormalTextStyle, TextItem.Text);
+		}
+		else
+		{
+			FinalText.AppendLine(TextItem.Text);
+		}
 	}
+
+	OverlayText->SetText(FinalText.ToText());
 }
 
 TSharedRef<SEditorViewport> SStaticMeshEditorViewport::GetViewportWidget()
@@ -201,6 +209,27 @@ bool SStaticMeshEditorViewport::PreviewComponentSelectionOverride(const UPrimiti
 	return false;
 }
 
+void SStaticMeshEditorViewport::ToggleShowNaniteProxy()
+{
+	if (PreviewMeshComponent)
+	{
+		FComponentReregisterContext ReregisterContext(PreviewMeshComponent);
+		PreviewMeshComponent->bDisplayNaniteProxyMesh = !PreviewMeshComponent->bDisplayNaniteProxyMesh;
+	}
+}
+
+bool SStaticMeshEditorViewport::IsShowNaniteProxyChecked() const
+{
+	return PreviewMeshComponent ? PreviewMeshComponent->bDisplayNaniteProxyMesh : false;
+}
+
+bool SStaticMeshEditorViewport::IsShowNaniteProxyVisible() const
+{
+	const UStaticMesh* PreviewStaticMesh = PreviewMeshComponent ? PreviewMeshComponent->GetStaticMesh() : nullptr;
+
+	return PreviewStaticMesh && PreviewStaticMesh->NaniteSettings.bEnabled ? true : false;
+}
+
 void SStaticMeshEditorViewport::UpdatePreviewSocketMeshes()
 {
 	UStaticMesh* const PreviewStaticMesh = PreviewMeshComponent ? PreviewMeshComponent->GetStaticMesh() : nullptr;
@@ -536,6 +565,13 @@ void SStaticMeshEditorViewport::BindCommands()
 
 	TSharedRef<FStaticMeshEditorViewportClient> EditorViewportClientRef = EditorViewportClient.ToSharedRef();
 
+	CommandList->MapAction(
+		Commands.SetShowNaniteProxy,
+		FExecuteAction::CreateSP(this, &SStaticMeshEditorViewport::ToggleShowNaniteProxy),
+		FCanExecuteAction(),
+		FIsActionChecked::CreateSP(this, &SStaticMeshEditorViewport::IsShowNaniteProxyChecked),
+		FIsActionButtonVisible::CreateSP(this, &SStaticMeshEditorViewport::IsShowNaniteProxyVisible));
+
 	CommandList->MapAction(
 		Commands.SetShowWireframe,
 		FExecuteAction::CreateSP( this, &SStaticMeshEditorViewport::SetViewModeWireframe ),
diff --git a/Engine/Source/Editor/StaticMeshEditor/Private/SStaticMeshEditorViewport.h b/Engine/Source/Editor/StaticMeshEditor/Private/SStaticMeshEditorViewport.h
--- a/Engine/Source/Editor/StaticMeshEditor/Private/SStaticMeshEditorViewport.h
+++ b/Engine/Source/Editor/StaticMeshEditor/Private/SStaticMeshEditorViewport.h
@@ -13,11 +13,13 @@
 #include "SEditorViewport.h"
 #include "SCommonEditorViewportToolbarBase.h"
 #include "SAssetEditorViewport.h"
+#include "Styling/StyleColors.h"
 
 class IStaticMeshEditor;
 class SVerticalBox;
 class UStaticMesh;
 class UStaticMeshComponent;
+class SRichTextBlock;
 
 /**
  * StaticMesh Editor Preview viewport widget
@@ -101,16 +103,17 @@ public:
 	/** Struct defining the text and its style of each item in the overlay widget */
 	struct FOverlayTextItem
 	{
-		explicit FOverlayTextItem(const FText& InText, const FName& InStyle = "TextBlock.ShadowedText")
-			: Text(InText), Style(InStyle)
+		explicit FOverlayTextItem(const FText& InText, bool bInIsWarning = false, bool bInIsCustomFormat=false)
+			: Text(InText), bIsWarning(bInIsWarning), bIsCustomFormat(bInIsCustomFormat)
 		{}
 
 		FText Text;
-		FName Style;
+		bool bIsWarning;
+		bool bIsCustomFormat;
 	};
 
 	/** Specifies an array of text items which will be added to the viewport overlay */
-	void PopulateOverlayText( const TArray<FOverlayTextItem>& TextItems );
+	void PopulateOverlayText( const TArrayView<FOverlayTextItem> TextItems );
 
 	// ICommonEditorViewportToolbarInfoProvider interface
 	virtual TSharedRef<class SEditorViewport> GetViewportWidget() override;
@@ -171,6 +174,12 @@ private:
 
 	/** Override for preview component selection to inform the editor we consider it selected. */
 	bool PreviewComponentSelectionOverride(const UPrimitiveComponent* InComponent) const;
+
+	void ToggleShowNaniteProxy();
+
+	bool IsShowNaniteProxyChecked() const;
+
+	bool IsShowNaniteProxyVisible() const;
 private:
 	
 	/** The parent tab where this viewport resides */
@@ -191,8 +200,8 @@ private:
 	/** The currently selected view mode. */
 	EViewModeIndex CurrentViewMode;
 
-	/** Pointer to the vertical box into which the overlay text items are added */
-	TSharedPtr<SVerticalBox> OverlayTextVerticalBox;
+	/** Pointer to the box into which the overlay text items are added */
+	TSharedPtr<SRichTextBlock> OverlayText;
 
 	/** Current LOD Selection where 0 is Auto */
 	int32 LODSelection;
diff --git a/Engine/Source/Editor/StaticMeshEditor/Private/SStaticMeshEditorViewportToolBar.cpp b/Engine/Source/Editor/StaticMeshEditor/Private/SStaticMeshEditorViewportToolBar.cpp
--- a/Engine/Source/Editor/StaticMeshEditor/Private/SStaticMeshEditorViewportToolBar.cpp
+++ b/Engine/Source/Editor/StaticMeshEditor/Private/SStaticMeshEditorViewportToolBar.cpp
@@ -43,24 +43,25 @@ TSharedRef<SWidget> SStaticMeshEditorViewportToolbar::GenerateShowMenu() const
 	{
 		auto Commands = FStaticMeshEditorCommands::Get();
 
+		ShowMenuBuilder.AddMenuEntry(Commands.SetShowNaniteProxy);
+
+		ShowMenuBuilder.BeginSection("MeshComponents", LOCTEXT("MeshComponments", "Mesh Components"));
 		ShowMenuBuilder.AddMenuEntry(Commands.SetShowSockets);
-		ShowMenuBuilder.AddMenuEntry(Commands.SetShowPivot);
 		ShowMenuBuilder.AddMenuEntry(Commands.SetShowVertices);
 		ShowMenuBuilder.AddMenuEntry(Commands.SetShowVertexColor);
-
+		ShowMenuBuilder.AddMenuEntry(Commands.SetShowNormals);
+		ShowMenuBuilder.AddMenuEntry(Commands.SetShowTangents);
+		ShowMenuBuilder.AddMenuEntry(Commands.SetShowBinormals);
 		ShowMenuBuilder.AddMenuSeparator();
 
+		ShowMenuBuilder.AddMenuEntry(Commands.SetShowPivot);
 		ShowMenuBuilder.AddMenuEntry(Commands.SetShowGrid);
 		ShowMenuBuilder.AddMenuEntry(Commands.SetShowBounds);
 		ShowMenuBuilder.AddMenuEntry(Commands.SetShowSimpleCollision);
 		ShowMenuBuilder.AddMenuEntry(Commands.SetShowComplexCollision);
 		ShowMenuBuilder.AddMenuEntry(Commands.SetShowPhysicalMaterialMasks);
 
-		ShowMenuBuilder.AddMenuSeparator();
 
- 		ShowMenuBuilder.AddMenuEntry(Commands.SetShowNormals);
-		ShowMenuBuilder.AddMenuEntry(Commands.SetShowTangents);
-		ShowMenuBuilder.AddMenuEntry(Commands.SetShowBinormals);
 
 		//ShowMenuBuilder.AddMenuSeparator();
 		//ShowMenuBuilder.AddMenuEntry(Commands.SetShowMeshEdges);
diff --git a/Engine/Source/Editor/StaticMeshEditor/Private/StaticMeshEditorActions.cpp b/Engine/Source/Editor/StaticMeshEditor/Private/StaticMeshEditorActions.cpp
--- a/Engine/Source/Editor/StaticMeshEditor/Private/StaticMeshEditorActions.cpp
+++ b/Engine/Source/Editor/StaticMeshEditor/Private/StaticMeshEditorActions.cpp
@@ -6,63 +6,64 @@
 
 void FStaticMeshEditorCommands::RegisterCommands()
 {
-	UI_COMMAND( SetShowWireframe, "Wireframe", "Toggles the viewmode of the Preview Pane between a lit view and a wireframe view.", EUserInterfaceActionType::ToggleButton, FInputChord() );
-	UI_COMMAND( SetShowVertexColor, "Vert Colors", "Toggles vertex colors.", EUserInterfaceActionType::ToggleButton, FInputChord() );
-	UI_COMMAND( SetShowPhysicalMaterialMasks, "Physical Material Masks", "Toggles physical material masks.", EUserInterfaceActionType::ToggleButton, FInputChord());
-	UI_COMMAND( SetDrawUVs, "UV", "Toggles display of the static mesh's UVs for the specified channel.", EUserInterfaceActionType::ToggleButton, FInputChord() );
-	UI_COMMAND( SetShowGrid, "Grid", "Displays the viewport grid.", EUserInterfaceActionType::ToggleButton, FInputChord() );
-	UI_COMMAND( SetShowBounds, "Bounds", "Toggles display of the bounds of the static mesh.", EUserInterfaceActionType::ToggleButton, FInputChord() );
-	UI_COMMAND( SetShowSimpleCollision, "Simple Collision", "Toggles display of the simplified collision mesh of the static mesh, if present.", EUserInterfaceActionType::ToggleButton, FInputChord());
-	UI_COMMAND( SetShowComplexCollision, "Complex Collision", "Toggles display of the complex collision for this static mesh.", EUserInterfaceActionType::ToggleButton, FInputChord());
-	UI_COMMAND( ResetCamera, "Reset Camera", "Resets the camera to focus on the static mesh", EUserInterfaceActionType::Button, FInputChord() );
-	UI_COMMAND( SetShowSockets, "Sockets", "Displays the static mesh sockets.", EUserInterfaceActionType::ToggleButton, FInputChord(EModifierKey::Alt, EKeys::S) );
-	UI_COMMAND( SetDrawAdditionalData, "Additional Data", "Draw additional user data associated with asset.", EUserInterfaceActionType::ToggleButton, FInputChord() );
+	UI_COMMAND(SetShowNaniteProxy, "Nanite Proxy", "Toggles the display of the Nanite proxy mesh. The proxy mesh is used when Nanite is not available such as platforms that do not support it.", EUserInterfaceActionType::ToggleButton, FInputChord(EModifierKey::Control, EKeys::N));
+	UI_COMMAND(SetShowWireframe, "Wireframe", "Toggles the viewmode of the Preview Pane between a lit view and a wireframe view.", EUserInterfaceActionType::ToggleButton, FInputChord());
+	UI_COMMAND(SetShowVertexColor, "Vert Colors", "Toggles vertex colors.", EUserInterfaceActionType::ToggleButton, FInputChord());
+	UI_COMMAND(SetShowPhysicalMaterialMasks, "Physical Material Masks", "Toggles physical material masks.", EUserInterfaceActionType::ToggleButton, FInputChord());
+	UI_COMMAND(SetDrawUVs, "UV", "Toggles display of the static mesh's UVs for the specified channel.", EUserInterfaceActionType::ToggleButton, FInputChord());
+	UI_COMMAND(SetShowGrid, "Grid", "Displays the viewport grid.", EUserInterfaceActionType::ToggleButton, FInputChord());
+	UI_COMMAND(SetShowBounds, "Bounds", "Toggles display of the bounds of the static mesh.", EUserInterfaceActionType::ToggleButton, FInputChord());
+	UI_COMMAND(SetShowSimpleCollision, "Simple Collision", "Toggles display of the simplified collision mesh of the static mesh, if present.", EUserInterfaceActionType::ToggleButton, FInputChord());
+	UI_COMMAND(SetShowComplexCollision, "Complex Collision", "Toggles display of the complex collision for this static mesh.", EUserInterfaceActionType::ToggleButton, FInputChord());
+	UI_COMMAND(ResetCamera, "Reset Camera", "Resets the camera to focus on the static mesh", EUserInterfaceActionType::Button, FInputChord());
+	UI_COMMAND(SetShowSockets, "Sockets", "Displays the static mesh sockets.", EUserInterfaceActionType::ToggleButton, FInputChord(EModifierKey::Alt, EKeys::S));
+	UI_COMMAND(SetDrawAdditionalData, "Additional Data", "Draw additional user data associated with asset.", EUserInterfaceActionType::ToggleButton, FInputChord());
 
-	UI_COMMAND( ReimportMesh, "Reimport Base Mesh", "Reimport the base mesh.", EUserInterfaceActionType::Button, FInputChord() );
-	UI_COMMAND( ReimportMeshWithNewFile, "Reimport Base Mesh With New File", "Reimport the base mesh using a new source file.", EUserInterfaceActionType::Button, FInputChord());
-	UI_COMMAND( ReimportAllMesh, "Reimport Base Mesh + LODs", "Reimport the base mesh and all the custom LODs.", EUserInterfaceActionType::Button, FInputChord());
-	UI_COMMAND( ReimportAllMeshWithNewFile, "Reimport Base Mesh + LODs With New File", "Reimport the base mesh using a new source file and all the custom LODs (No new source file for the custom LODs).", EUserInterfaceActionType::Button, FInputChord() );
+	UI_COMMAND(ReimportMesh, "Reimport Base Mesh", "Reimport the base mesh.", EUserInterfaceActionType::Button, FInputChord());
+	UI_COMMAND(ReimportMeshWithNewFile, "Reimport Base Mesh With New File", "Reimport the base mesh using a new source file.", EUserInterfaceActionType::Button, FInputChord());
+	UI_COMMAND(ReimportAllMesh, "Reimport Base Mesh + LODs", "Reimport the base mesh and all the custom LODs.", EUserInterfaceActionType::Button, FInputChord());
+	UI_COMMAND(ReimportAllMeshWithNewFile, "Reimport Base Mesh + LODs With New File", "Reimport the base mesh using a new source file and all the custom LODs (No new source file for the custom LODs).", EUserInterfaceActionType::Button, FInputChord());
 
-	UI_COMMAND( ToggleShowNormals, "Normals", "Toggles display of vertex normals in the Preview Pane.", EUserInterfaceActionType::ToggleButton, FInputChord());
-	UI_COMMAND( ToggleShowTangents, "Tangents", "Toggles display of vertex tangents in the Preview Pane.", EUserInterfaceActionType::ToggleButton, FInputChord());
-	UI_COMMAND( ToggleShowBinormals, "Binormals", "Toggles display of vertex binormals (orthogonal vector to normal and tangent) in the Preview Pane.", EUserInterfaceActionType::ToggleButton, FInputChord());
-	UI_COMMAND( ToggleShowVertices, "Vertices", "Toggles display of vertices in the Preview Pane.", EUserInterfaceActionType::ToggleButton, FInputChord());
-	UI_COMMAND( ToggleShowPivots, "Show Pivot", "Display the pivot location of the static mesh.", EUserInterfaceActionType::ToggleButton, FInputChord());
-	UI_COMMAND( ToggleShowGrids, "Grid", "Displays the viewport grid.", EUserInterfaceActionType::ToggleButton, FInputChord());
-	UI_COMMAND( ToggleShowBounds, "Bounds", "Toggles display of the bounds of the static mesh.", EUserInterfaceActionType::ToggleButton, FInputChord());
-	UI_COMMAND( ToggleShowSimpleCollisions, "Simple Collision", "Toggles display of the simplified collision mesh of the static mesh, if present.", EUserInterfaceActionType::ToggleButton, FInputChord());
-	UI_COMMAND( ToggleShowComplexCollisions, "Complex Collision", "Toggles display of the complex collision for this static mesh.", EUserInterfaceActionType::ToggleButton, FInputChord());
-	UI_COMMAND( ToggleShowSockets, "Sockets", "Displays the static mesh sockets.", EUserInterfaceActionType::ToggleButton, FInputChord());
-	UI_COMMAND( ToggleShowWireframes, "Wireframe", "Toggles the viewmode of the Preview Pane between a lit view and a wireframe view.", EUserInterfaceActionType::ToggleButton, FInputChord());
-	UI_COMMAND( ToggleShowVertexColors, "Vert Colors", "Toggles vertex colors.", EUserInterfaceActionType::ToggleButton, FInputChord());
+	UI_COMMAND(ToggleShowNormals, "Normals", "Toggles display of vertex normals in the Preview Pane.", EUserInterfaceActionType::ToggleButton, FInputChord());
+	UI_COMMAND(ToggleShowTangents, "Tangents", "Toggles display of vertex tangents in the Preview Pane.", EUserInterfaceActionType::ToggleButton, FInputChord());
+	UI_COMMAND(ToggleShowBinormals, "Binormals", "Toggles display of vertex binormals (orthogonal vector to normal and tangent) in the Preview Pane.", EUserInterfaceActionType::ToggleButton, FInputChord());
+	UI_COMMAND(ToggleShowVertices, "Vertices", "Toggles display of vertices in the Preview Pane.", EUserInterfaceActionType::ToggleButton, FInputChord());
+	UI_COMMAND(ToggleShowPivots, "Show Pivot", "Display the pivot location of the static mesh.", EUserInterfaceActionType::ToggleButton, FInputChord());
+	UI_COMMAND(ToggleShowGrids, "Grid", "Displays the viewport grid.", EUserInterfaceActionType::ToggleButton, FInputChord());
+	UI_COMMAND(ToggleShowBounds, "Bounds", "Toggles display of the bounds of the static mesh.", EUserInterfaceActionType::ToggleButton, FInputChord());
+	UI_COMMAND(ToggleShowSimpleCollisions, "Simple Collision", "Toggles display of the simplified collision mesh of the static mesh, if present.", EUserInterfaceActionType::ToggleButton, FInputChord());
+	UI_COMMAND(ToggleShowComplexCollisions, "Complex Collision", "Toggles display of the complex collision for this static mesh.", EUserInterfaceActionType::ToggleButton, FInputChord());
+	UI_COMMAND(ToggleShowSockets, "Sockets", "Displays the static mesh sockets.", EUserInterfaceActionType::ToggleButton, FInputChord());
+	UI_COMMAND(ToggleShowWireframes, "Wireframe", "Toggles the viewmode of the Preview Pane between a lit view and a wireframe view.", EUserInterfaceActionType::ToggleButton, FInputChord());
+	UI_COMMAND(ToggleShowVertexColors, "Vert Colors", "Toggles vertex colors.", EUserInterfaceActionType::ToggleButton, FInputChord());
 
-	UI_COMMAND( SetShowNormals, "Normals", "Toggles display of vertex normals in the Preview Pane.", EUserInterfaceActionType::ToggleButton, FInputChord(EModifierKey::Alt, EKeys::N) );
-	UI_COMMAND( SetShowTangents, "Tangents", "Toggles display of vertex tangents in the Preview Pane.", EUserInterfaceActionType::ToggleButton, FInputChord(EModifierKey::Alt, EKeys::T) );
-	UI_COMMAND( SetShowBinormals, "Binormals", "Toggles display of vertex binormals (orthogonal vector to normal and tangent) in the Preview Pane.", EUserInterfaceActionType::ToggleButton, FInputChord(EModifierKey::Alt, EKeys::B) );
-	UI_COMMAND( SetShowVertices, "Vertices", "Toggles display of vertices in the Preview Pane.", EUserInterfaceActionType::ToggleButton, FInputChord(EModifierKey::Alt, EKeys::V) ); 
-	UI_COMMAND( SetShowPivot, "Show Pivot", "Display the pivot location of the static mesh.", EUserInterfaceActionType::ToggleButton, FInputChord(EModifierKey::Alt, EKeys::P) );
+	UI_COMMAND(SetShowNormals, "Normals", "Toggles display of vertex normals in the Preview Pane.", EUserInterfaceActionType::ToggleButton, FInputChord(EModifierKey::Alt, EKeys::N));
+	UI_COMMAND(SetShowTangents, "Tangents", "Toggles display of vertex tangents in the Preview Pane.", EUserInterfaceActionType::ToggleButton, FInputChord(EModifierKey::Alt, EKeys::T));
+	UI_COMMAND(SetShowBinormals, "Binormals", "Toggles display of vertex binormals (orthogonal vector to normal and tangent) in the Preview Pane.", EUserInterfaceActionType::ToggleButton, FInputChord(EModifierKey::Alt, EKeys::B));
+	UI_COMMAND(SetShowVertices, "Vertices", "Toggles display of vertices in the Preview Pane.", EUserInterfaceActionType::ToggleButton, FInputChord(EModifierKey::Alt, EKeys::V));
+	UI_COMMAND(SetShowPivot, "Show Pivot", "Display the pivot location of the static mesh.", EUserInterfaceActionType::ToggleButton, FInputChord(EModifierKey::Alt, EKeys::P));
 
-	UI_COMMAND( CreateDOP10X, "Add 10DOP-X Simplified Collision", "Generates a new axis-aligned box collision mesh with the 4 X-axis aligned edges beveled (10 total sides) encompassing the static mesh.", EUserInterfaceActionType::Button, FInputChord() );
-	UI_COMMAND( CreateDOP10Y, "Add 10DOP-Y Simplified Collision", "Generates a new axis-aligned box collision mesh with the 4 Y-axis aligned edges beveled (10 total sides) encompassing the static mesh.", EUserInterfaceActionType::Button, FInputChord() );
-	UI_COMMAND( CreateDOP10Z, "Add 10DOP-Z Simplified Collision", "Generates a new axis-aligned box collision mesh with the 4 Z-axis aligned edges beveled (10 total sides) encompassing the static mesh.", EUserInterfaceActionType::Button, FInputChord() );
-	UI_COMMAND( CreateDOP18, "Add 18DOP Simplified Collision", "Generates a new axis-aligned box collision mesh with all edges beveled (18 total sides) encompassing the static mesh.", EUserInterfaceActionType::Button, FInputChord() );
-	UI_COMMAND( CreateDOP26, "Add 26DOP Simplified Collision", "Generates a new axis-aligned box collision mesh with all edges and corners beveled (26 total sides) encompassing the static mesh.", EUserInterfaceActionType::Button, FInputChord() );
+	UI_COMMAND(CreateDOP10X, "Add 10DOP-X Simplified Collision", "Generates a new axis-aligned box collision mesh with the 4 X-axis aligned edges beveled (10 total sides) encompassing the static mesh.", EUserInterfaceActionType::Button, FInputChord());
+	UI_COMMAND(CreateDOP10Y, "Add 10DOP-Y Simplified Collision", "Generates a new axis-aligned box collision mesh with the 4 Y-axis aligned edges beveled (10 total sides) encompassing the static mesh.", EUserInterfaceActionType::Button, FInputChord());
+	UI_COMMAND(CreateDOP10Z, "Add 10DOP-Z Simplified Collision", "Generates a new axis-aligned box collision mesh with the 4 Z-axis aligned edges beveled (10 total sides) encompassing the static mesh.", EUserInterfaceActionType::Button, FInputChord());
+	UI_COMMAND(CreateDOP18, "Add 18DOP Simplified Collision", "Generates a new axis-aligned box collision mesh with all edges beveled (18 total sides) encompassing the static mesh.", EUserInterfaceActionType::Button, FInputChord());
+	UI_COMMAND(CreateDOP26, "Add 26DOP Simplified Collision", "Generates a new axis-aligned box collision mesh with all edges and corners beveled (26 total sides) encompassing the static mesh.", EUserInterfaceActionType::Button, FInputChord());
 
-	UI_COMMAND( CreateBoxCollision, "Add Box Simplified Collision", "Generates a new box collision mesh encompassing the static mesh.", EUserInterfaceActionType::Button, FInputChord() );
-	UI_COMMAND( CreateSphereCollision, "Add Sphere Simplified Collision", "Generates a new sphere collision mesh encompassing the static mesh.", EUserInterfaceActionType::Button, FInputChord() );
-	UI_COMMAND( CreateSphylCollision, "Add Capsule Simplified Collision", "Generates a new capsule collision mesh encompassing the static mesh.", EUserInterfaceActionType::Button, FInputChord() );
-	UI_COMMAND( CreateAutoConvexCollision, "Auto Convex Collision", "Opens the Auto Convex Collision Tool for generating a new convex collision mesh, or meshes.", EUserInterfaceActionType::Button, FInputChord() );
-	UI_COMMAND( RemoveCollision, "Remove Collision", "Removes any simplified collision assigned to the static mesh.", EUserInterfaceActionType::Button, FInputChord() );
-	UI_COMMAND( ConvertBoxesToConvex, "Convert Boxes to Convex", "Converts any simple box collision meshes to convex collision meshes.", EUserInterfaceActionType::Button, FInputChord() );
+	UI_COMMAND(CreateBoxCollision, "Add Box Simplified Collision", "Generates a new box collision mesh encompassing the static mesh.", EUserInterfaceActionType::Button, FInputChord());
+	UI_COMMAND(CreateSphereCollision, "Add Sphere Simplified Collision", "Generates a new sphere collision mesh encompassing the static mesh.", EUserInterfaceActionType::Button, FInputChord());
+	UI_COMMAND(CreateSphylCollision, "Add Capsule Simplified Collision", "Generates a new capsule collision mesh encompassing the static mesh.", EUserInterfaceActionType::Button, FInputChord());
+	UI_COMMAND(CreateAutoConvexCollision, "Auto Convex Collision", "Opens the Auto Convex Collision Tool for generating a new convex collision mesh, or meshes.", EUserInterfaceActionType::Button, FInputChord());
+	UI_COMMAND(RemoveCollision, "Remove Collision", "Removes any simplified collision assigned to the static mesh.", EUserInterfaceActionType::Button, FInputChord());
+	UI_COMMAND(ConvertBoxesToConvex, "Convert Boxes to Convex", "Converts any simple box collision meshes to convex collision meshes.", EUserInterfaceActionType::Button, FInputChord());
 
-	UI_COMMAND( CopyCollisionFromSelectedMesh, "Copy Collision from Selected Static Mesh", "Copy collision from the static mesh selected in Content Browser.", EUserInterfaceActionType::Button, FInputChord() );
+	UI_COMMAND(CopyCollisionFromSelectedMesh, "Copy Collision from Selected Static Mesh", "Copy collision from the static mesh selected in Content Browser.", EUserInterfaceActionType::Button, FInputChord());
 
-	UI_COMMAND( FindSource, "Find Source", "Opens explorer at the location of this asset.", EUserInterfaceActionType::Button, FInputChord() );
-	
-	UI_COMMAND( ChangeMesh, "Change Mesh", "Changes the static mesh asset loaded in the Static Mesh Editor to the asset currently selected in the Content Browser.", EUserInterfaceActionType::Button, FInputChord() );
+	UI_COMMAND(FindSource, "Find Source", "Opens explorer at the location of this asset.", EUserInterfaceActionType::Button, FInputChord());
 
-	UI_COMMAND( SaveGeneratedLODs, "Save Generated LODs", "Saves the generated LODs in the mesh package.", EUserInterfaceActionType::Button, FInputChord());
-	UI_COMMAND( BakeMaterials, "Bake Out Materials", "Bake out Materials for given LOD(s).", EUserInterfaceActionType::Button, FInputChord());
+	UI_COMMAND(ChangeMesh, "Change Mesh", "Changes the static mesh asset loaded in the Static Mesh Editor to the asset currently selected in the Content Browser.", EUserInterfaceActionType::Button, FInputChord());
+
+	UI_COMMAND(SaveGeneratedLODs, "Save Generated LODs", "Saves the generated LODs in the mesh package.", EUserInterfaceActionType::Button, FInputChord());
+	UI_COMMAND(BakeMaterials, "Bake Out Materials", "Bake out Materials for given LOD(s).", EUserInterfaceActionType::Button, FInputChord());
 }
 
 #undef LOCTEXT_NAMESPACE
diff --git a/Engine/Source/Editor/StaticMeshEditor/Private/StaticMeshEditorTools.cpp b/Engine/Source/Editor/StaticMeshEditor/Private/StaticMeshEditorTools.cpp
--- a/Engine/Source/Editor/StaticMeshEditor/Private/StaticMeshEditorTools.cpp
+++ b/Engine/Source/Editor/StaticMeshEditor/Private/StaticMeshEditorTools.cpp
@@ -47,6 +47,7 @@
 #include "Widgets/Input/SFilePathPicker.h"
 #include "Widgets/Input/STextComboBox.h"
 #include "PerPlatformPropertyCustomization.h"
+#include "Misc/ScopedSlowTask.h"
 
 const uint32 MaxHullCount = 64;
 const uint32 MinHullCount = 2;
@@ -4686,19 +4687,17 @@ void FNaniteSettingsLayout::ApplyChanges()
 	UStaticMesh* StaticMesh = StaticMeshEditor.GetStaticMesh();
 	check(StaticMesh);
 
-	// Calling Begin and EndSlowTask are rather dangerous because they tick
-	// Slate. Call them here and flush rendering commands to be sure!.
-
-	FFormatNamedArguments Args;
-	Args.Add(TEXT("StaticMeshName"), FText::FromString(StaticMesh->GetName()));
-	GWarn->BeginSlowTask(FText::Format(LOCTEXT("ApplyNaniteChanges", "Applying changes to {StaticMeshName}..."), Args), true);
-	FlushRenderingCommands();
+	{
+		FFormatNamedArguments Args;
+		Args.Add(TEXT("StaticMeshName"), FText::FromString(StaticMesh->GetName()));
+		FScopedSlowTask SlowTask(0, FText::Format(LOCTEXT("ApplyNaniteChanges", "Applying changes to {StaticMeshName}..."), Args), true);
+		SlowTask.MakeDialog();
 
-	StaticMesh->Modify();
-	StaticMesh->NaniteSettings = NaniteSettings;
-	StaticMesh->PostEditChange();
+		StaticMesh->Modify();
+		StaticMesh->NaniteSettings = NaniteSettings;
+		StaticMesh->PostEditChange();
+	}
 
-	GWarn->EndSlowTask();
 
 	StaticMeshEditor.RefreshTool();
 }
diff --git a/Engine/Source/Editor/StaticMeshEditor/Private/StaticMeshEditorViewportClient.cpp b/Engine/Source/Editor/StaticMeshEditor/Private/StaticMeshEditorViewportClient.cpp
--- a/Engine/Source/Editor/StaticMeshEditor/Private/StaticMeshEditorViewportClient.cpp
+++ b/Engine/Source/Editor/StaticMeshEditor/Private/StaticMeshEditorViewportClient.cpp
@@ -785,7 +785,7 @@ void FStaticMeshEditorViewportClient::DrawCanvas( FViewport& InViewport, FSceneV
 		}
 	}
 
-	TArray<SStaticMeshEditorViewport::FOverlayTextItem> TextItems;
+	TArray<SStaticMeshEditorViewport::FOverlayTextItem, TInlineAllocator<10>> TextItems;
 
 	const int32 CurrentLODLevel = [this, &StaticMeshEditor, &View]()
 	{
@@ -796,17 +796,19 @@ void FStaticMeshEditorViewportClient::DrawCanvas( FViewport& InViewport, FSceneV
 			LOD - 1;
 	}();
 
-	const int32 CurrentMinLODLevel = StaticMesh->GetMinLOD().GetValue();
-	const bool bBelowMinLOD = CurrentLODLevel < CurrentMinLODLevel;
-	TextItems.Add(SStaticMeshEditorViewport::FOverlayTextItem(
-		FText::Format(NSLOCTEXT("UnrealEd", "LOD_F", "LOD:  {0}"), FText::AsNumber(CurrentLODLevel)),
-		bBelowMinLOD ? "TextBlock.ShadowedTextWarning" : "TextBlock.ShadowedText"));
-	
-	if ( bBelowMinLOD )
+	if (StaticMesh->NaniteSettings.bEnabled)
+	{
+		TextItems.Emplace(FText::Format(NSLOCTEXT("UnrealEd", "NaniteEnabled", "<TextBlock.ShadowedText>Nanite Enabled</> <TextBlock.ShadowedTextWarning>{0}</>"), StaticMeshComponent->bDisplayNaniteProxyMesh ? NSLOCTEXT("UnrealEd", "ShowingNaniteProxy", "(Showing Proxy)") : FText::GetEmpty()), false, true);
+	}
+	else
 	{
-		TextItems.Add(SStaticMeshEditorViewport::FOverlayTextItem(
-			FText::Format(NSLOCTEXT("UnrealEd", "BelowMinLODWarning_F", "Selected LOD is below the minimum of {0}"),
-				FText::AsNumber(CurrentMinLODLevel)), "TextBlock.ShadowedTextWarning"));
+		const int32 CurrentMinLODLevel = StaticMesh->GetMinLOD().GetValue();
+		const bool bBelowMinLOD = CurrentLODLevel < CurrentMinLODLevel;
+
+		if ( bBelowMinLOD )
+		{
+			TextItems.Emplace(FText::Format(NSLOCTEXT("UnrealEd", "BelowMinLODWarning_F", "Selected LOD is below the minimum of {0}"),FText::AsNumber(CurrentMinLODLevel)), true);
+		}
 	}
 
 	const float CurrentScreenSize = ComputeBoundsScreenSize(StaticMeshComponent->Bounds.Origin, StaticMeshComponent->Bounds.SphereRadius, View);
@@ -814,17 +816,13 @@ void FStaticMeshEditorViewportClient::DrawCanvas( FViewport& InViewport, FSceneV
 	FormatOptions.MinimumFractionalDigits = 3;
 	FormatOptions.MaximumFractionalDigits = 6;
 	FormatOptions.MaximumIntegralDigits = 6;
-	TextItems.Add(SStaticMeshEditorViewport::FOverlayTextItem(
-		FText::Format(NSLOCTEXT("UnrealEd", "ScreenSize_F", "Current Screen Size:  {0}"), FText::AsNumber(CurrentScreenSize, &FormatOptions))));
+	TextItems.Emplace(FText::Format(NSLOCTEXT("UnrealEd", "ScreenSize_F", "Current Screen Size:  {0}"), FText::AsNumber(CurrentScreenSize, &FormatOptions)));
 
-	TextItems.Add(SStaticMeshEditorViewport::FOverlayTextItem(
-		FText::Format(NSLOCTEXT("UnrealEd", "Triangles_F", "Triangles:  {0}"), FText::AsNumber(StaticMeshEditorPtr.Pin()->GetNumTriangles(CurrentLODLevel)))));
+	TextItems.Emplace(FText::Format(NSLOCTEXT("UnrealEd", "Triangles_F", "Triangles:  {0}"), FText::AsNumber(StaticMeshEditorPtr.Pin()->GetNumTriangles(CurrentLODLevel))));
 
-	TextItems.Add(SStaticMeshEditorViewport::FOverlayTextItem(
-		FText::Format(NSLOCTEXT("UnrealEd", "Vertices_F", "Vertices:  {0}"), FText::AsNumber(StaticMeshEditorPtr.Pin()->GetNumVertices(CurrentLODLevel)))));
+	TextItems.Emplace(FText::Format(NSLOCTEXT("UnrealEd", "Vertices_F", "Vertices:  {0}"), FText::AsNumber(StaticMeshEditorPtr.Pin()->GetNumVertices(CurrentLODLevel))));
 
-	TextItems.Add(SStaticMeshEditorViewport::FOverlayTextItem(
-		FText::Format(NSLOCTEXT("UnrealEd", "UVChannels_F", "UV Channels:  {0}"), FText::AsNumber(StaticMeshEditorPtr.Pin()->GetNumUVChannels(CurrentLODLevel)))));
+	TextItems.Emplace(FText::Format(NSLOCTEXT("UnrealEd", "UVChannels_F", "UV Channels:  {0}"), FText::AsNumber(StaticMeshEditorPtr.Pin()->GetNumUVChannels(CurrentLODLevel))));
 
 	if(StaticMesh->GetRenderData() && StaticMesh->GetRenderData()->LODResources.Num() > 0 )
 	{
@@ -844,31 +842,29 @@ void FStaticMeshEditorViewportClient::DrawCanvas( FViewport& InViewport, FSceneV
 				NumberOptions.MinimumFractionalDigits = 2;
 				NumberOptions.MaximumFractionalDigits = 2;
 
-				TextItems.Add(SStaticMeshEditorViewport::FOverlayTextItem(
-					FText::Format(NSLOCTEXT("UnrealEd", "DistanceFieldRes_F", "Distance Field:  {0}x{1}x{2} = {3}Mb"), FText::AsNumber(VolumeData.Size.X), FText::AsNumber(VolumeData.Size.Y), FText::AsNumber(VolumeData.Size.Z), FText::AsNumber(MemoryMb, &NumberOptions))));
+				TextItems.Emplace(FText::Format(NSLOCTEXT("UnrealEd", "DistanceFieldRes_F", "Distance Field:  {0}x{1}x{2} = {3}Mb"), FText::AsNumber(VolumeData.Size.X), FText::AsNumber(VolumeData.Size.Y), FText::AsNumber(VolumeData.Size.Z), FText::AsNumber(MemoryMb, &NumberOptions)));
 			}
 		}
 	}
 
-	TextItems.Add(SStaticMeshEditorViewport::FOverlayTextItem(
+	TextItems.Emplace(
 		FText::Format(NSLOCTEXT("UnrealEd", "ApproxSize_F", "Approx Size: {0}x{1}x{2}"),
 		FText::AsNumber(int32(StaticMesh->GetBounds().BoxExtent.X * 2.0f)), // x2 as artists wanted length not radius
 		FText::AsNumber(int32(StaticMesh->GetBounds().BoxExtent.Y * 2.0f)),
-		FText::AsNumber(int32(StaticMesh->GetBounds().BoxExtent.Z * 2.0f)))));
+		FText::AsNumber(int32(StaticMesh->GetBounds().BoxExtent.Z * 2.0f))));
 
 	// Show the number of collision primitives
 	if(StaticMesh->GetBodySetup())
 	{
-		TextItems.Add(SStaticMeshEditorViewport::FOverlayTextItem(
-			FText::Format(NSLOCTEXT("UnrealEd", "NumPrimitives_F", "Num Collision Primitives:  {0}"), FText::AsNumber(StaticMesh->GetBodySetup()->AggGeom.GetElementCount()))));
+		TextItems.Emplace(FText::Format(NSLOCTEXT("UnrealEd", "NumPrimitives_F", "Num Collision Primitives:  {0}"), FText::AsNumber(StaticMesh->GetBodySetup()->AggGeom.GetElementCount())));
 	}
 
 	if (StaticMeshComponent && StaticMeshComponent->SectionIndexPreview != INDEX_NONE)
 	{
-		TextItems.Add(SStaticMeshEditorViewport::FOverlayTextItem(NSLOCTEXT("UnrealEd", "MeshSectionsHiddenWarning",  "Mesh Sections Hidden")));
+		TextItems.Emplace(NSLOCTEXT("UnrealEd", "MeshSectionsHiddenWarning",  "Mesh Sections Hidden"));
 	}
 
-	StaticMeshEditorViewport->PopulateOverlayText(TextItems);
+	StaticMeshEditorViewport->PopulateOverlayText(MakeArrayView(TextItems));
 
  	int32 X = Canvas.GetRenderTarget()->GetSizeXY().X - 300;
  	int32 Y = 30;
diff --git a/Engine/Source/Editor/StaticMeshEditor/Public/StaticMeshEditorActions.h b/Engine/Source/Editor/StaticMeshEditor/Public/StaticMeshEditorActions.h
--- a/Engine/Source/Editor/StaticMeshEditor/Public/StaticMeshEditorActions.h
+++ b/Engine/Source/Editor/StaticMeshEditor/Public/StaticMeshEditorActions.h
@@ -16,81 +16,82 @@ class FStaticMeshEditorCommands : public TCommands<FStaticMeshEditorCommands>
 
 public:
 	FStaticMeshEditorCommands() : TCommands<FStaticMeshEditorCommands>
-	(
-		"StaticMeshEditor", // Context name for fast lookup
-		NSLOCTEXT("Contexts", "StaticMeshEditor", "StaticMesh Editor"), // Localized context name for displaying
-		"EditorViewport",  // Parent
-		FEditorStyle::GetStyleSetName() // Icon Style Set
-	)
+		(
+			"StaticMeshEditor", // Context name for fast lookup
+			NSLOCTEXT("Contexts", "StaticMeshEditor", "StaticMesh Editor"), // Localized context name for displaying
+			"EditorViewport",  // Parent
+			FEditorStyle::GetStyleSetName() // Icon Style Set
+			)
 	{
 	}
-	
+
 	/**
 	 * StaticMesh Editor Commands
 	 */
-	
-	/**  */
-	TSharedPtr< FUICommandInfo > SetShowWireframe;
-	TSharedPtr< FUICommandInfo > SetShowVertexColor;
-	TSharedPtr< FUICommandInfo > SetShowPhysicalMaterialMasks;
-	TSharedPtr< FUICommandInfo > SetDrawUVs;
-	TSharedPtr< FUICommandInfo > SetShowGrid;
-	TSharedPtr< FUICommandInfo > SetShowBounds;
-	TSharedPtr< FUICommandInfo > SetShowSimpleCollision;
-	TSharedPtr< FUICommandInfo > SetShowComplexCollision;
-	TSharedPtr< FUICommandInfo > ResetCamera;
-	TSharedPtr< FUICommandInfo > SetShowSockets;
-	TSharedPtr< FUICommandInfo > SetDrawAdditionalData;
+
+	 /**  */
+	TSharedPtr<FUICommandInfo> SetShowNaniteProxy;
+	TSharedPtr<FUICommandInfo> SetShowWireframe;
+	TSharedPtr<FUICommandInfo> SetShowVertexColor;
+	TSharedPtr<FUICommandInfo> SetShowPhysicalMaterialMasks;
+	TSharedPtr<FUICommandInfo> SetDrawUVs;
+	TSharedPtr<FUICommandInfo> SetShowGrid;
+	TSharedPtr<FUICommandInfo> SetShowBounds;
+	TSharedPtr<FUICommandInfo> SetShowSimpleCollision;
+	TSharedPtr<FUICommandInfo> SetShowComplexCollision;
+	TSharedPtr<FUICommandInfo> ResetCamera;
+	TSharedPtr<FUICommandInfo> SetShowSockets;
+	TSharedPtr<FUICommandInfo> SetDrawAdditionalData;
 
 	// Mesh toolbar Commands
-	TSharedPtr< FUICommandInfo > ReimportMesh;
-	TSharedPtr< FUICommandInfo > ReimportMeshWithNewFile;
-	TSharedPtr< FUICommandInfo > ReimportAllMesh;
-	TSharedPtr< FUICommandInfo > ReimportAllMeshWithNewFile;
+	TSharedPtr<FUICommandInfo> ReimportMesh;
+	TSharedPtr<FUICommandInfo> ReimportMeshWithNewFile;
+	TSharedPtr<FUICommandInfo> ReimportAllMesh;
+	TSharedPtr<FUICommandInfo> ReimportAllMeshWithNewFile;
 
 	// toolbar commands
-	TSharedPtr< FUICommandInfo > ToggleShowNormals;
-	TSharedPtr< FUICommandInfo > ToggleShowTangents;
-	TSharedPtr< FUICommandInfo > ToggleShowBinormals;
-	TSharedPtr< FUICommandInfo > ToggleShowPivots;
-	TSharedPtr< FUICommandInfo > ToggleShowVertices;
-	TSharedPtr< FUICommandInfo > ToggleShowGrids;
-	TSharedPtr< FUICommandInfo > ToggleShowBounds;
-	TSharedPtr< FUICommandInfo > ToggleShowSimpleCollisions;
-	TSharedPtr< FUICommandInfo > ToggleShowComplexCollisions;
-	TSharedPtr< FUICommandInfo > ToggleShowSockets;
-	TSharedPtr< FUICommandInfo > ToggleShowWireframes;
-	TSharedPtr< FUICommandInfo > ToggleShowVertexColors;
+	TSharedPtr<FUICommandInfo> ToggleShowNormals;
+	TSharedPtr<FUICommandInfo> ToggleShowTangents;
+	TSharedPtr<FUICommandInfo> ToggleShowBinormals;
+	TSharedPtr<FUICommandInfo> ToggleShowPivots;
+	TSharedPtr<FUICommandInfo> ToggleShowVertices;
+	TSharedPtr<FUICommandInfo> ToggleShowGrids;
+	TSharedPtr<FUICommandInfo> ToggleShowBounds;
+	TSharedPtr<FUICommandInfo> ToggleShowSimpleCollisions;
+	TSharedPtr<FUICommandInfo> ToggleShowComplexCollisions;
+	TSharedPtr<FUICommandInfo> ToggleShowSockets;
+	TSharedPtr<FUICommandInfo> ToggleShowWireframes;
+	TSharedPtr<FUICommandInfo> ToggleShowVertexColors;
 
 	// View Menu Commands
-	TSharedPtr< FUICommandInfo > SetShowNormals;
-	TSharedPtr< FUICommandInfo > SetShowTangents;
-	TSharedPtr< FUICommandInfo > SetShowBinormals;
-	TSharedPtr< FUICommandInfo > SetShowPivot;
-	TSharedPtr< FUICommandInfo > SetShowVertices;
+	TSharedPtr<FUICommandInfo> SetShowNormals;
+	TSharedPtr<FUICommandInfo> SetShowTangents;
+	TSharedPtr<FUICommandInfo> SetShowBinormals;
+	TSharedPtr<FUICommandInfo> SetShowPivot;
+	TSharedPtr<FUICommandInfo> SetShowVertices;
 
 	// Collision Menu Commands
-	TSharedPtr< FUICommandInfo > CreateDOP10X;
-	TSharedPtr< FUICommandInfo > CreateDOP10Y;
-	TSharedPtr< FUICommandInfo > CreateDOP10Z;
-	TSharedPtr< FUICommandInfo > CreateDOP18;
-	TSharedPtr< FUICommandInfo > CreateDOP26;
-	TSharedPtr< FUICommandInfo > CreateBoxCollision;
-	TSharedPtr< FUICommandInfo > CreateSphereCollision;
-	TSharedPtr< FUICommandInfo > CreateSphylCollision;
-	TSharedPtr< FUICommandInfo > CreateAutoConvexCollision;
-	TSharedPtr< FUICommandInfo > RemoveCollision;
-	TSharedPtr< FUICommandInfo > ConvertBoxesToConvex;
-	TSharedPtr< FUICommandInfo > CopyCollisionFromSelectedMesh;
+	TSharedPtr<FUICommandInfo> CreateDOP10X;
+	TSharedPtr<FUICommandInfo> CreateDOP10Y;
+	TSharedPtr<FUICommandInfo> CreateDOP10Z;
+	TSharedPtr<FUICommandInfo> CreateDOP18;
+	TSharedPtr<FUICommandInfo> CreateDOP26;
+	TSharedPtr<FUICommandInfo> CreateBoxCollision;
+	TSharedPtr<FUICommandInfo> CreateSphereCollision;
+	TSharedPtr<FUICommandInfo> CreateSphylCollision;
+	TSharedPtr<FUICommandInfo> CreateAutoConvexCollision;
+	TSharedPtr<FUICommandInfo> RemoveCollision;
+	TSharedPtr<FUICommandInfo> ConvertBoxesToConvex;
+	TSharedPtr<FUICommandInfo> CopyCollisionFromSelectedMesh;
 
 	// Mesh Menu Commands
-	TSharedPtr< FUICommandInfo > FindSource;
+	TSharedPtr<FUICommandInfo> FindSource;
 
-	TSharedPtr< FUICommandInfo > ChangeMesh;
+	TSharedPtr<FUICommandInfo> ChangeMesh;
 
-	TSharedPtr< FUICommandInfo > BakeMaterials;
+	TSharedPtr<FUICommandInfo> BakeMaterials;
 
-	TSharedPtr< FUICommandInfo > SaveGeneratedLODs;
+	TSharedPtr<FUICommandInfo> SaveGeneratedLODs;
 
 	/**
 	 * Initialize commands
diff --git a/Engine/Source/Runtime/Engine/Classes/Components/StaticMeshComponent.h b/Engine/Source/Runtime/Engine/Classes/Components/StaticMeshComponent.h
--- a/Engine/Source/Runtime/Engine/Classes/Components/StaticMeshComponent.h
+++ b/Engine/Source/Runtime/Engine/Classes/Components/StaticMeshComponent.h
@@ -315,6 +315,9 @@ public:
 	UPROPERTY(transient)
 	uint8 bDisplayPhysicalMaterialMasks : 1;
 
+	/** For nanite enabled meshes, we'll only show the proxy mesh if this is true */
+	UPROPERTY()
+	uint8 bDisplayNaniteProxyMesh:1;
 #endif
 
 	/**
diff --git a/Engine/Source/Runtime/Engine/Private/Components/StaticMeshComponent.cpp b/Engine/Source/Runtime/Engine/Private/Components/StaticMeshComponent.cpp
--- a/Engine/Source/Runtime/Engine/Private/Components/StaticMeshComponent.cpp
+++ b/Engine/Source/Runtime/Engine/Private/Components/StaticMeshComponent.cpp
@@ -182,6 +182,7 @@ UStaticMeshComponent::UStaticMeshComponent(const FObjectInitializer& ObjectIniti
 	bCustomOverrideVertexColorPerLOD = false;
 	bDisplayVertexColors = false;
 	bDisplayPhysicalMaterialMasks = false;
+	bDisplayNaniteProxyMesh = false;
 #endif
 }
 
diff --git a/Engine/Source/Runtime/Engine/Private/StaticMeshRender.cpp b/Engine/Source/Runtime/Engine/Private/StaticMeshRender.cpp
--- a/Engine/Source/Runtime/Engine/Private/StaticMeshRender.cpp
+++ b/Engine/Source/Runtime/Engine/Private/StaticMeshRender.cpp
@@ -2363,8 +2363,15 @@ FPrimitiveSceneProxy* UStaticMeshComponent::CreateSceneProxy()
 		return nullptr;
 	}
 
+	// Whether or not to allow nanite for this component
+#if WITH_EDITORONLY_DATA
+	const bool bAllowNanite = !bDisplayNaniteProxyMesh;
+#else
+	const bool bAllowNanite = true;
+#endif
+
 	// TODO: Abstract with a common helper
-	if (UseNanite(GetScene()->GetShaderPlatform()) && GetStaticMesh()->GetRenderData()->NaniteResources.PageStreamingStates.Num())
+	if (bAllowNanite && UseNanite(GetScene()->GetShaderPlatform()) && GetStaticMesh()->GetRenderData()->NaniteResources.PageStreamingStates.Num())
 	{
 		LLM_SCOPE(ELLMTag::StaticMesh);
 		return ::new Nanite::FSceneProxy(this);
