Better SDF stats in static mesh editor

diff --git a/Engine/Source/Editor/StaticMeshEditor/Private/StaticMeshEditorViewportClient.cpp b/Engine/Source/Editor/StaticMeshEditor/Private/StaticMeshEditorViewportClient.cpp
--- a/Engine/Source/Editor/StaticMeshEditor/Private/StaticMeshEditorViewportClient.cpp
+++ b/Engine/Source/Editor/StaticMeshEditor/Private/StaticMeshEditorViewportClient.cpp
@@ -831,13 +831,14 @@ void FStaticMeshEditorViewportClient::DrawCanvas( FViewport& InViewport, FSceneV
 			const FDistanceFieldVolumeData& VolumeData = *(StaticMesh->GetRenderData()->LODResources[0].DistanceFieldData);
 			const FIntVector VolumeSize = VolumeData.Mips[0].IndirectionDimensions * DistanceField::UniqueDataBrickSize;
 			{
-				float MemoryMb = VolumeData.GetResourceSizeBytes() / (1024.0f * 1024.0f);
+				float AlwaysLoadedMemoryMb = VolumeData.GetResourceSizeBytes() / (1024.0f * 1024.0f);
+				float HighestResMipMemoryMb = VolumeData.Mips[0].BulkSize / (1024.0f * 1024.0f);
 
 				FNumberFormattingOptions NumberOptions;
 				NumberOptions.MinimumFractionalDigits = 2;
 				NumberOptions.MaximumFractionalDigits = 2;
 
-				TextItems.Emplace(FText::Format(NSLOCTEXT("UnrealEd", "DistanceFieldRes_F", "Distance Field:  {0}x{1}x{2} = {3}Mb"), FText::AsNumber(VolumeSize.X), FText::AsNumber(VolumeSize.Y), FText::AsNumber(VolumeSize.Z), FText::AsNumber(MemoryMb, &NumberOptions)));
+				TextItems.Emplace(FText::Format(NSLOCTEXT("UnrealEd", "DistanceFieldRes_F", "Distance Field:  {0}x{1}x{2} = {3}Mb always loaded, {4}Mb streamed"), FText::AsNumber(VolumeSize.X), FText::AsNumber(VolumeSize.Y), FText::AsNumber(VolumeSize.Z), FText::AsNumber(AlwaysLoadedMemoryMb, &NumberOptions), FText::AsNumber(HighestResMipMemoryMb, &NumberOptions)));
 			}
 		}
 	}
