Changing the static mesh editor's additional asset menu options to check the asset editor context

diff --git a/Engine/Source/Editor/StaticMeshEditor/Private/StaticMeshEditor.cpp b/Engine/Source/Editor/StaticMeshEditor/Private/StaticMeshEditor.cpp
--- a/Engine/Source/Editor/StaticMeshEditor/Private/StaticMeshEditor.cpp
+++ b/Engine/Source/Editor/StaticMeshEditor/Private/StaticMeshEditor.cpp
@@ -51,6 +51,7 @@
 #include "RawMesh.h"
 #include "EditorViewportTabContent.h"
 #include "EditorViewportLayout.h"
+#include "Toolkits/AssetEditorToolkitMenuContext.h"
 
 #define LOCTEXT_NAMESPACE "StaticMeshEditor"
 
@@ -463,11 +464,20 @@ void FStaticMeshEditor::ExtendMenu()
 	AddMenuExtender(StaticMeshEditorModule->GetMenuExtensibilityManager()->GetAllExtenders(GetToolkitCommands(), GetEditingObjects()));
 
 	UToolMenu* AssetMenu = UToolMenus::Get()->ExtendMenu("MainFrame.MainMenu.Asset");
-	FToolMenuSection& AssetSection = AssetMenu->AddSection("StaticMeshAssets");
-	AssetSection.AddMenuEntry(FStaticMeshEditorCommands::Get().FindSource);
-	AssetSection.AddMenuEntry(FStaticMeshEditorCommands::Get().SetDrawAdditionalData );
-	AssetSection.AddMenuEntry(FStaticMeshEditorCommands::Get().BakeMaterials);
-	AssetSection.AddDynamicEntry("SaveGeneratedLODs", FNewToolMenuSectionDelegate::CreateStatic(&Local::FillMeshMenu));
+	FToolMenuSection& AssetSection = AssetMenu->FindOrAddSection("AssetEditorActions");
+	FToolMenuEntry& Entry = AssetSection.AddDynamicEntry("AssetManagerEditorStaticMeshCommands", FNewToolMenuSectionDelegate::CreateLambda([this](FToolMenuSection& InSection)
+		{
+			UAssetEditorToolkitMenuContext* MenuContext = InSection.FindContext<UAssetEditorToolkitMenuContext>();
+			if (MenuContext && MenuContext->Toolkit.IsValid() && MenuContext->Toolkit.Pin()->GetToolkitFName() == GetToolkitFName())
+			{
+				InSection.AddMenuEntry(FStaticMeshEditorCommands::Get().FindSource);
+				InSection.AddMenuEntry(FStaticMeshEditorCommands::Get().SetDrawAdditionalData);
+				InSection.AddMenuEntry(FStaticMeshEditorCommands::Get().BakeMaterials);
+				InSection.AddDynamicEntry("SaveGeneratedLODs", FNewToolMenuSectionDelegate::CreateStatic(&Local::FillMeshMenu));
+			}
+		}
+	));
+
 
 }
 
