Allow creating of a single Convex collision in static mesh editor

diff --git a/Engine/Source/Editor/StaticMeshEditor/Private/StaticMeshEditorTools.cpp b/Engine/Source/Editor/StaticMeshEditor/Private/StaticMeshEditorTools.cpp
--- a/Engine/Source/Editor/StaticMeshEditor/Private/StaticMeshEditorTools.cpp
+++ b/Engine/Source/Editor/StaticMeshEditor/Private/StaticMeshEditorTools.cpp
@@ -51,7 +51,7 @@
 #include "Misc/ScopedSlowTask.h"
 
 const uint32 MaxHullCount = 64;
-const uint32 MinHullCount = 2;
+const uint32 MinHullCount = 1;
 const uint32 DefaultHullCount = 4;
 const uint32 HullCountDelta = 1;
 
