Fixed an issue where the current LOD would not display on static meshes

#okfforgithub public

diff --git a/Engine/Source/Editor/StaticMeshEditor/Private/StaticMeshEditorViewportClient.cpp b/Engine/Source/Editor/StaticMeshEditor/Private/StaticMeshEditorViewportClient.cpp
--- a/Engine/Source/Editor/StaticMeshEditor/Private/StaticMeshEditorViewportClient.cpp
+++ b/Engine/Source/Editor/StaticMeshEditor/Private/StaticMeshEditorViewportClient.cpp
@@ -813,9 +813,12 @@ void FStaticMeshEditorViewportClient::DrawCanvas( FViewport& InViewport, FSceneV
 	}
 	else
 	{
+		
 		const int32 CurrentMinLODLevel = StaticMesh->GetMinLOD().GetValue();
 		const bool bBelowMinLOD = CurrentLODLevel < CurrentMinLODLevel;
 
+		TextItems.Emplace(FText::Format(NSLOCTEXT("UnrealEd", "LOD_F", "LOD:  {0}"), FText::AsNumber(CurrentLODLevel)), bBelowMinLOD);
+
 		if ( bBelowMinLOD )
 		{
 			TextItems.Emplace(FText::Format(NSLOCTEXT("UnrealEd", "BelowMinLODWarning_F", "Selected LOD is below the minimum of {0}"),FText::AsNumber(CurrentMinLODLevel)), true);
