Change how the lambda captures the asset editor toolkit name.

diff --git a/Engine/Source/Editor/StaticMeshEditor/Private/StaticMeshEditor.cpp b/Engine/Source/Editor/StaticMeshEditor/Private/StaticMeshEditor.cpp
--- a/Engine/Source/Editor/StaticMeshEditor/Private/StaticMeshEditor.cpp
+++ b/Engine/Source/Editor/StaticMeshEditor/Private/StaticMeshEditor.cpp
@@ -464,10 +464,11 @@ void FStaticMeshEditor::ExtendMenu()
 
 	UToolMenu* AssetMenu = UToolMenus::Get()->ExtendMenu("MainFrame.MainMenu.Asset");
 	FToolMenuSection& AssetSection = AssetMenu->FindOrAddSection("AssetEditorActions");
-	FToolMenuEntry& Entry = AssetSection.AddDynamicEntry("AssetManagerEditorStaticMeshCommands", FNewToolMenuSectionDelegate::CreateLambda([this](FToolMenuSection& InSection)
+	FName StaticMeshToolkitName = GetToolkitFName();
+	FToolMenuEntry& Entry = AssetSection.AddDynamicEntry("AssetManagerEditorStaticMeshCommands", FNewToolMenuSectionDelegate::CreateLambda([StaticMeshToolkitName](FToolMenuSection& InSection)
 		{
 			UAssetEditorToolkitMenuContext* MenuContext = InSection.FindContext<UAssetEditorToolkitMenuContext>();
-			if (MenuContext && MenuContext->Toolkit.IsValid() && MenuContext->Toolkit.Pin()->GetToolkitFName() == GetToolkitFName())
+			if (MenuContext && MenuContext->Toolkit.IsValid() && MenuContext->Toolkit.Pin()->GetToolkitFName() == StaticMeshToolkitName)
 			{
 				InSection.AddMenuEntry(FStaticMeshEditorCommands::Get().FindSource);
 				InSection.AddMenuEntry(FStaticMeshEditorCommands::Get().SetDrawAdditionalData);
