More static mesh editor cosmetic cleanup

diff --git a/Engine/Source/Editor/StaticMeshEditor/Private/StaticMeshEditorTools.cpp b/Engine/Source/Editor/StaticMeshEditor/Private/StaticMeshEditorTools.cpp
--- a/Engine/Source/Editor/StaticMeshEditor/Private/StaticMeshEditorTools.cpp
+++ b/Engine/Source/Editor/StaticMeshEditor/Private/StaticMeshEditorTools.cpp
@@ -3188,6 +3188,7 @@ void FLevelOfDetailSettingsLayout::AddToDetailsPanel( IDetailLayoutBuilder& Deta
 		.Text(LOCTEXT("LODGroup", "LOD Group"))
 	]
 	.ValueContent()
+	.VAlign(VAlign_Center)
 	[
 		SAssignNew(LODGroupComboBox, STextComboBox)
 		.Font(IDetailLayoutBuilder::GetDetailFont())
@@ -3198,20 +3199,21 @@ void FLevelOfDetailSettingsLayout::AddToDetailsPanel( IDetailLayoutBuilder& Deta
 	
 	LODSettingsCategory.AddCustomRow( LOCTEXT("LODImport", "LOD Import") )
 	.RowTag("LODImport")
-		.NameContent()
-		[
-			SNew(STextBlock)
-			.Font( IDetailLayoutBuilder::GetDetailFont() )
-			.Text(LOCTEXT("LODImport", "LOD Import"))
-		]
+	.NameContent()
+	[
+		SNew(STextBlock)
+		.Font(IDetailLayoutBuilder::GetDetailFont())
+		.Text(LOCTEXT("LODImport", "LOD Import"))
+	]
 	.ValueContent()
-		[
-			SNew(STextComboBox)
-			.Font(IDetailLayoutBuilder::GetDetailFont())
-			.OptionsSource(&LODNames)
-			.InitiallySelectedItem(LODNames[0])
-			.OnSelectionChanged(this, &FLevelOfDetailSettingsLayout::OnImportLOD)
-		];
+	.VAlign(VAlign_Center)
+	[
+		SNew(STextComboBox)
+		.Font(IDetailLayoutBuilder::GetDetailFont())
+		.OptionsSource(&LODNames)
+		.InitiallySelectedItem(LODNames[0])
+		.OnSelectionChanged(this, &FLevelOfDetailSettingsLayout::OnImportLOD)
+	];
 
 	int32 PlatformNumber = PlatformInfo::GetAllPlatformGroupNames().Num();
 
@@ -3283,7 +3285,7 @@ void FLevelOfDetailSettingsLayout::AddToDetailsPanel( IDetailLayoutBuilder& Deta
 		.Font( IDetailLayoutBuilder::GetDetailFont() )
 		.Text(LOCTEXT("NumberOfLODs", "Number of LODs"))
 	]
-	.ValueContent()
+	.ValueContent().VAlign(VAlign_Center)
 	[
 		SNew(SSpinBox<int32>)
 		.Font( IDetailLayoutBuilder::GetDetailFont() )
@@ -3316,6 +3318,7 @@ void FLevelOfDetailSettingsLayout::AddToDetailsPanel( IDetailLayoutBuilder& Deta
 	.RowTag("ApplyChanges")
 	.ValueContent()
 	.HAlign(HAlign_Left)
+	.VAlign(VAlign_Center)
 	[
 		SNew(SButton)
 		.OnClicked(this, &FLevelOfDetailSettingsLayout::OnApply)
@@ -3421,6 +3424,7 @@ void FLevelOfDetailSettingsLayout::AddLODLevelCategories( IDetailLayoutBuilder&
 			.IsEnabled(this, &FLevelOfDetailSettingsLayout::IsLodComboBoxEnabledForLodPicker)
 		]
 		.ValueContent()
+		.VAlign(VAlign_Center)
 		[
 			OnGenerateLodComboBoxForLodPicker()
 		];
@@ -3595,6 +3599,7 @@ void FLevelOfDetailSettingsLayout::AddLODLevelCategories( IDetailLayoutBuilder&
 					.Text(LOCTEXT("SourceImportFilenameName", "Source Import Filename"))
 				]
 				.ValueContent()
+					.VAlign(VAlign_Center)
 					.MinDesiredWidth(125.0f)
 					.MaxDesiredWidth(0.0f)
 				[
@@ -4758,6 +4763,7 @@ void FNaniteSettingsLayout::AddToDetailsPanel(IDetailLayoutBuilder& DetailBuilde
 			.ToolTipText(LOCTEXT("PositionPrecisionTooltip", "Precision of vertex positions."))
 		]
 		.ValueContent()
+		.VAlign(VAlign_Center)
 		[
 			SAssignNew(TerminationCriterionCombo, STextComboBox)
 			.Font(IDetailLayoutBuilder::GetDetailFont())
@@ -4779,6 +4785,7 @@ void FNaniteSettingsLayout::AddToDetailsPanel(IDetailLayoutBuilder& DetailBuilde
 			.ToolTipText(LOCTEXT("ProxyTrianglePercentTooltip", "Percentage of triangles to reduce down to for generating a coarse proxy\nmesh that will be used anywhere the full detail Nanite data can't,\nincluding platforms that don't support Nanite rendering."))
 		]
 		.ValueContent()
+		.VAlign(VAlign_Center)
 		[
 			SNew(SSpinBox<float>)
 			.Font(IDetailLayoutBuilder::GetDetailFont())
@@ -4793,6 +4800,7 @@ void FNaniteSettingsLayout::AddToDetailsPanel(IDetailLayoutBuilder& DetailBuilde
 	NaniteSettingsCategory.AddCustomRow(LOCTEXT("ApplyChanges", "Apply Changes"))
 	.ValueContent()
 	.HAlign(HAlign_Left)
+	.VAlign(VAlign_Center)
 	[
 		SNew(SButton)
 		.OnClicked(this, &FNaniteSettingsLayout::OnApply)
