Improve scalability for the way we strip LOD on static mesh.

Change the per platform overrides with  a quality level property  driven by scalability/device profile ini files. User can now enter a Quality level override and specify the qualtiy in the platform engine or scalability ini files. this allow to have different values on one platform, for example XboxOne Base vs XboxOne X. Keep the PerPlatform MinLOD code logic for licensee but allow to choose per instance to use the new system.

diff --git a/Engine/Plugins/Enterprise/DataprepEditor/Source/DataprepEditor/Private/Widgets/SDataprepEditorViewport.cpp b/Engine/Plugins/Enterprise/DataprepEditor/Source/DataprepEditor/Private/Widgets/SDataprepEditorViewport.cpp
--- a/Engine/Plugins/Enterprise/DataprepEditor/Source/DataprepEditor/Private/Widgets/SDataprepEditorViewport.cpp
+++ b/Engine/Plugins/Enterprise/DataprepEditor/Source/DataprepEditor/Private/Widgets/SDataprepEditorViewport.cpp
@@ -1577,7 +1577,7 @@ FPrimitiveSceneProxy* UCustomStaticMeshComponent::CreateSceneProxy()
 	}
 
 	const FStaticMeshLODResourcesArray& LODResources = GetStaticMesh()->GetRenderData()->LODResources;
-	if (LODResources.Num() == 0	|| LODResources[FMath::Clamp<int32>(GetStaticMesh()->GetMinLOD().Default, 0, LODResources.Num()-1)].VertexBuffers.StaticMeshVertexBuffer.GetNumVertices() == 0)
+	if (LODResources.Num() == 0	|| LODResources[FMath::Clamp<int32>(GetStaticMesh()->GetDefaultMinLOD(), 0, LODResources.Num()-1)].VertexBuffers.StaticMeshVertexBuffer.GetNumVertices() == 0)
 	{
 		return nullptr;
 	}
diff --git a/Engine/Plugins/Experimental/GPULightmass/Source/GPULightmass/Private/Scene/Scene.cpp b/Engine/Plugins/Experimental/GPULightmass/Source/GPULightmass/Private/Scene/Scene.cpp
--- a/Engine/Plugins/Experimental/GPULightmass/Source/GPULightmass/Private/Scene/Scene.cpp
+++ b/Engine/Plugins/Experimental/GPULightmass/Source/GPULightmass/Private/Scene/Scene.cpp
@@ -688,7 +688,7 @@ void FScene::AddGeometryInstanceFromComponent(UStaticMeshComponent* InComponent)
 
 	RegisteredStaticMeshComponentUObjects.Add(InComponent, Instance);
 
-	const int32 SMCurrentMinLOD = InComponent->GetStaticMesh()->GetMinLOD().Default;
+	const int32 SMCurrentMinLOD = InComponent->GetStaticMesh()->GetDefaultMinLOD();
 	int32 EffectiveMinLOD = InComponent->bOverrideMinLOD ? InComponent->MinLOD : SMCurrentMinLOD;
 
 	// Find the first LOD with any vertices (ie that haven't been stripped)
diff --git a/Engine/Plugins/Experimental/PanoramicCapture/Source/PanoramicCapture/Private/StereoStaticMeshComponent.cpp b/Engine/Plugins/Experimental/PanoramicCapture/Source/PanoramicCapture/Private/StereoStaticMeshComponent.cpp
--- a/Engine/Plugins/Experimental/PanoramicCapture/Source/PanoramicCapture/Private/StereoStaticMeshComponent.cpp
+++ b/Engine/Plugins/Experimental/PanoramicCapture/Source/PanoramicCapture/Private/StereoStaticMeshComponent.cpp
@@ -57,7 +57,7 @@ FPrimitiveSceneProxy* UStereoStaticMeshComponent::CreateSceneProxy()
     if ((GetStaticMesh() == nullptr) ||
 		(GetStaticMesh()->GetRenderData() == nullptr) ||
 		(GetStaticMesh()->GetRenderData()->LODResources.Num() == 0) ||
-		(GetStaticMesh()->GetRenderData()->LODResources[GetStaticMesh()->GetMinLOD().GetValue()].VertexBuffers.PositionVertexBuffer.GetNumVertices() == 0))
+		(GetStaticMesh()->GetRenderData()->LODResources[GetStaticMesh()->GetMinLODIdx()].VertexBuffers.PositionVertexBuffer.GetNumVertices() == 0))
     {
         return nullptr;
     }
diff --git a/Engine/Plugins/FX/Niagara/Source/Niagara/Private/NiagaraRendererMeshes.cpp b/Engine/Plugins/FX/Niagara/Source/Niagara/Private/NiagaraRendererMeshes.cpp
--- a/Engine/Plugins/FX/Niagara/Source/Niagara/Private/NiagaraRendererMeshes.cpp
+++ b/Engine/Plugins/FX/Niagara/Source/Niagara/Private/NiagaraRendererMeshes.cpp
@@ -179,7 +179,7 @@ void FNiagaraRendererMeshes::Initialize(const UNiagaraRendererProperties* InProp
 			MeshData.PivotOffset = MeshProperties.PivotOffset;
 			MeshData.PivotOffsetSpace = MeshProperties.PivotOffsetSpace;
 			MeshData.Scale = MeshProperties.Scale;
-			MeshData.MinimumLOD = Mesh->GetMinLOD().GetValue();
+			MeshData.MinimumLOD = MeshProperties.Mesh->GetMinLODIdx();
 
 			// Create an index remap from mesh material index to it's index in the master material list
 			TArray<UMaterialInterface*> MeshMaterials;
diff --git a/Engine/Source/Editor/StaticMeshEditor/Private/StaticMeshEditorTools.cpp b/Engine/Source/Editor/StaticMeshEditor/Private/StaticMeshEditorTools.cpp
--- a/Engine/Source/Editor/StaticMeshEditor/Private/StaticMeshEditorTools.cpp
+++ b/Engine/Source/Editor/StaticMeshEditor/Private/StaticMeshEditorTools.cpp
@@ -29,6 +29,7 @@
 #include "FbxMeshUtils.h"
 #include "Widgets/Input/SVectorInputBox.h"
 #include "Widgets/Input/SNumericEntryBox.h"
+#include "SPerQualityLevelPropertiesWidget.h"
 #include "PlatformInfo.h"
 #include "ContentStreaming.h"
 #include "EditorDirectories.h"
@@ -3207,12 +3208,15 @@ void FLevelOfDetailSettingsLayout::AddToDetailsPanel( IDetailLayoutBuilder& Deta
 			.OnSelectionChanged(this, &FLevelOfDetailSettingsLayout::OnImportLOD)
 		];
 
+	int32 PlatformNumber = PlatformInfo::GetAllPlatformGroupNames().Num();
+	bool bDisablePerPlatformMinLod = GEngine->UsePerQualityLevelProperty && StaticMesh->GetQualityLevelMinLOD().bIsEnabled;
+	
 	{
 		TAttribute<TArray<FName>> PlatformOverrideNames = TAttribute<TArray<FName>>::Create(TAttribute<TArray<FName>>::FGetter::CreateSP(this, &FLevelOfDetailSettingsLayout::GetMinLODPlatformOverrideNames));
 		FPerPlatformPropertyCustomNodeBuilderArgs Args;
 		Args.NameWidget =
 			SNew(STextBlock)
-			.IsEnabled(FLevelOfDetailSettingsLayout::GetLODCount() > 1)
+			.IsEnabled(FLevelOfDetailSettingsLayout::GetLODCount() > 1 && !bDisablePerPlatformMinLod)
 			.Font(IDetailLayoutBuilder::GetDetailFont())
 			.Text(LOCTEXT("MinLOD", "Minimum LOD"));
 
@@ -3225,6 +3229,39 @@ void FLevelOfDetailSettingsLayout::AddToDetailsPanel( IDetailLayoutBuilder& Deta
 		LODSettingsCategory.AddCustomBuilder(MakeShared<FPerPlatformPropertyCustomNodeBuilder>(MoveTemp(Args)));
 	}
 
+	LODSettingsCategory.AddCustomRow(LOCTEXT("QualityLevelMinLOD", "Quality Level Min LOD"))
+		.NameContent()
+		[
+			SNew(STextBlock)
+			.Font(IDetailLayoutBuilder::GetDetailFont())
+		.Text(LOCTEXT("QualityLevelMinLOD", "Quality Level Min LOD"))
+		]
+	.ValueContent()
+		.MinDesiredWidth((float)(StaticMesh->GetQualityLevelMinLOD().PerQuality.Num() + 1)*125.0f)
+		.MaxDesiredWidth((float)((int32)QualityLevelProperty::EQualityLevels::Num + 1)*125.0f)
+		[
+			SNew(SPerQualityLevelPropertiesWidget)
+			.IsEnabled(FLevelOfDetailSettingsLayout::GetLODCount() > 1 && bDisablePerPlatformMinLod)
+		.OnGenerateWidget(this, &FLevelOfDetailSettingsLayout::GetMinQualityLevelLODWidget)
+		.OnAddEntry(this, &FLevelOfDetailSettingsLayout::AddMinLODQualityLevelOverride)
+		.OnRemoveEntry(this, &FLevelOfDetailSettingsLayout::RemoveMinLODQualityLevelOverride)
+		.EntryNames(this, &FLevelOfDetailSettingsLayout::GetMinQualityLevelLODOverrideNames)
+		];
+
+	LODSettingsCategory.AddCustomRow(LOCTEXT("QualityLevelMinLODEnable", "Enable"))
+		[
+			SNew(SCheckBox)
+			.IsChecked(this, &FLevelOfDetailSettingsLayout::IsMinLODQualityLevelChecked)
+			.OnCheckStateChanged(this, &FLevelOfDetailSettingsLayout::OnMinLODQualityLevelChecked)
+		.HAlign(HAlign_Right)
+			.Padding(FMargin(4.0f, 0.0f, 0.0f, 0.0f))
+			.Content()
+			[
+				SNew(STextBlock)
+				.Font(IDetailLayoutBuilder::GetDetailFont())
+				.Text(LOCTEXT("QualityLevelMinLODEnable", "Enable"))
+			]
+		];
 	{
 		TAttribute<TArray<FName>> PlatformOverrideNames = TAttribute<TArray<FName>>::Create(TAttribute<TArray<FName>>::FGetter::CreateSP(this, &FLevelOfDetailSettingsLayout::GetNumStreamedLODsPlatformOverrideNames));
 		FPerPlatformPropertyCustomNodeBuilderArgs Args;
@@ -3914,6 +3951,29 @@ void FLevelOfDetailSettingsLayout::OnLODGroupChanged(TSharedPtr<FString> NewValu
 	}
 }
 
+bool FLevelOfDetailSettingsLayout::IsMinLODQualityLevelEnabled() const
+{
+	UStaticMesh* StaticMesh = StaticMeshEditor.GetStaticMesh();
+	check(StaticMesh);
+	return StaticMesh->GetQualityLevelMinLOD().bIsEnabled;
+}
+
+ECheckBoxState FLevelOfDetailSettingsLayout::IsMinLODQualityLevelChecked() const
+{
+	return IsMinLODQualityLevelEnabled() ? ECheckBoxState::Checked : ECheckBoxState::Unchecked;
+}
+
+void FLevelOfDetailSettingsLayout::OnMinLODQualityLevelChecked(ECheckBoxState NewState)
+{
+	UStaticMesh* StaticMesh = StaticMeshEditor.GetStaticMesh();
+	check(StaticMesh);
+
+	FPerQualityLevelInt PerQualityLevel = StaticMesh->GetQualityLevelMinLOD();
+	PerQualityLevel.bIsEnabled = (NewState == ECheckBoxState::Checked) ? true : false;
+	StaticMesh->SetQualityLevelMinLOD(PerQualityLevel);
+	StaticMeshEditor.RefreshTool();
+}
+
 bool FLevelOfDetailSettingsLayout::IsAutoLODEnabled() const
 {
 	UStaticMesh* StaticMesh = StaticMeshEditor.GetStaticMesh();
@@ -4215,6 +4275,114 @@ TArray<FName> FLevelOfDetailSettingsLayout::GetMinLODPlatformOverrideNames() con
 	return KeyArray;
 }
 
+int32 FLevelOfDetailSettingsLayout::GetMinQualityLevelLOD(FName QualityLevel) const
+{
+	UStaticMesh* StaticMesh = StaticMeshEditor.GetStaticMesh();
+	check(StaticMesh);
+
+	int32 QLKey = QualityLevelProperty::FNameToQualityLevel(QualityLevel);
+	const int32* ValuePtr = (QualityLevel == NAME_None) ? nullptr : StaticMesh->GetQualityLevelMinLOD().PerQuality.Find(QLKey);
+	return (ValuePtr != nullptr) ? *ValuePtr : StaticMesh->GetQualityLevelMinLOD().Default;
+}
+
+void FLevelOfDetailSettingsLayout::OnMinQualityLevelLODChanged(int32 NewValue, FName QualityLevel)
+{
+	UStaticMesh* StaticMesh = StaticMeshEditor.GetStaticMesh();
+	check(StaticMesh);
+
+	{
+		FStaticMeshComponentRecreateRenderStateContext ReregisterContext(StaticMesh, false);
+		NewValue = FMath::Clamp<int32>(NewValue, 0, MAX_STATIC_MESH_LODS - 1);
+		FPerQualityLevelInt MinLOD = StaticMesh->GetQualityLevelMinLOD();
+		int32 QLKey = QualityLevelProperty::FNameToQualityLevel(QualityLevel);
+		if (QualityLevel == NAME_None || QLKey == INDEX_NONE)
+		{
+			MinLOD.Default = NewValue;
+		}
+		else
+		{
+			int32* ValuePtr = MinLOD.PerQuality.Find(QLKey);
+			if (ValuePtr != nullptr)
+			{
+				*ValuePtr = NewValue;
+			}
+		}
+		StaticMesh->SetQualityLevelMinLOD(MoveTemp(MinLOD));
+		StaticMesh->Modify();
+	}
+	StaticMeshEditor.RefreshViewport();
+}
+
+void FLevelOfDetailSettingsLayout::OnMinQualityLevelLODCommitted(int32 InValue, ETextCommit::Type CommitInfo, FName QualityLevel)
+{
+	OnMinQualityLevelLODChanged(InValue, QualityLevel);
+}
+
+TSharedRef<SWidget> FLevelOfDetailSettingsLayout::GetMinQualityLevelLODWidget(FName QualityLevelName) const
+{
+	return SNew(SSpinBox<int32>)
+		.Font(IDetailLayoutBuilder::GetDetailFont())
+		.Value(this, &FLevelOfDetailSettingsLayout::GetMinQualityLevelLOD, QualityLevelName)
+		.OnValueChanged(const_cast<FLevelOfDetailSettingsLayout*>(this), &FLevelOfDetailSettingsLayout::OnMinQualityLevelLODChanged, QualityLevelName)
+		.OnValueCommitted(const_cast<FLevelOfDetailSettingsLayout*>(this), &FLevelOfDetailSettingsLayout::OnMinQualityLevelLODCommitted, QualityLevelName)
+		.MinValue(0)
+		.MaxValue(MAX_STATIC_MESH_LODS)
+		.ToolTipText(this, &FLevelOfDetailSettingsLayout::GetMinLODTooltip)
+		.IsEnabled(FLevelOfDetailSettingsLayout::GetLODCount() > 1);
+}
+
+bool FLevelOfDetailSettingsLayout::AddMinLODQualityLevelOverride(FName QualityLevelName)
+{
+	FScopedTransaction Transaction(LOCTEXT("AddMinLODQualityLevelOverride", "Add Min LOD Quality Level Override"));
+	UStaticMesh* StaticMesh = StaticMeshEditor.GetStaticMesh();
+	check(StaticMesh);
+	StaticMesh->Modify();
+	int32 QLKey = QualityLevelProperty::FNameToQualityLevel(QualityLevelName);
+	if (StaticMesh->GetQualityLevelMinLOD().PerQuality.Find(QLKey) == nullptr)
+	{
+		FPerQualityLevelInt MinLOD = StaticMesh->GetQualityLevelMinLOD();
+		float Value = MinLOD.Default;
+		MinLOD.PerQuality.Add(QLKey, Value);
+		StaticMesh->SetQualityLevelMinLOD(MoveTemp(MinLOD));
+		OnMinQualityLevelLODChanged(Value, QualityLevelName);
+		return true;
+	}
+	return false;
+}
+
+bool FLevelOfDetailSettingsLayout::RemoveMinLODQualityLevelOverride(FName QualityLevelName)
+{
+	FScopedTransaction Transaction(LOCTEXT("RemoveMinLODQualityLevelOverride", "Remove Min LOD Quality Level Override"));
+	UStaticMesh* StaticMesh = StaticMeshEditor.GetStaticMesh();
+	check(StaticMesh);
+	StaticMesh->Modify();
+
+	FPerQualityLevelInt MinLOD = StaticMesh->GetQualityLevelMinLOD();
+	int32 QL = QualityLevelProperty::FNameToQualityLevel(QualityLevelName);
+	if (QL != INDEX_NONE && MinLOD.PerQuality.Remove(QL) != 0)
+	{
+		float Value = MinLOD.Default;
+		StaticMesh->SetQualityLevelMinLOD(MoveTemp(MinLOD));
+		OnMinQualityLevelLODChanged(Value, QualityLevelName);
+		return true;
+	}
+	return false;
+}
+
+TArray<FName> FLevelOfDetailSettingsLayout::GetMinQualityLevelLODOverrideNames() const
+{
+	UStaticMesh* StaticMesh = StaticMeshEditor.GetStaticMesh();
+	check(StaticMesh);
+
+	TArray<FName> OverrideNames;
+	for (const TPair<int32, int32>& Pair : StaticMesh->GetQualityLevelMinLOD().PerQuality)
+	{
+		OverrideNames.Add(QualityLevelProperty::QualityLevelToFName(Pair.Key));
+	}
+	OverrideNames.Sort(FNameLexicalLess());
+	return OverrideNames;
+}
+
 /** @return - whether value was different */
 static bool UpdateStaticMeshNumStreamedLODsHelper(UStaticMesh* StaticMesh, int32 NewValue, FName Platform)
 {
diff --git a/Engine/Source/Editor/StaticMeshEditor/Private/StaticMeshEditorTools.h b/Engine/Source/Editor/StaticMeshEditor/Private/StaticMeshEditorTools.h
--- a/Engine/Source/Editor/StaticMeshEditor/Private/StaticMeshEditorTools.h
+++ b/Engine/Source/Editor/StaticMeshEditor/Private/StaticMeshEditorTools.h
@@ -549,6 +549,17 @@ private:
 	bool RemoveMinLODPlatformOverride(FName PlatformGroupName);
 	TArray<FName> GetMinLODPlatformOverrideNames() const;
 
+	void OnMinQualityLevelLODChanged(int32 NewValue, FName QualityLevel);
+	void OnMinQualityLevelLODCommitted(int32 InValue, ETextCommit::Type CommitInfo, FName QualityLevel);
+	int32 GetMinQualityLevelLOD(FName QualityLevel) const;
+	TSharedRef<SWidget> GetMinQualityLevelLODWidget(FName QualityLevelName) const;
+	bool AddMinLODQualityLevelOverride(FName QualityLevelName);
+	bool RemoveMinLODQualityLevelOverride(FName QualityLevelName);
+	TArray<FName> GetMinQualityLevelLODOverrideNames() const;
+	ECheckBoxState IsMinLODQualityLevelChecked() const;
+	void OnMinLODQualityLevelChecked(ECheckBoxState NewState);
+	bool IsMinLODQualityLevelEnabled() const;
+
 	void OnNumStreamedLODsChanged(int32 NewValue, FName Platform);
 	void OnNumStreamedLODsCommitted(int32 InValue, ETextCommit::Type CommitInfo, FName Platform);
 	int32 GetNumStreamedLODs(FName Platform) const;
diff --git a/Engine/Source/Editor/StaticMeshEditor/Private/StaticMeshEditorViewportClient.cpp b/Engine/Source/Editor/StaticMeshEditor/Private/StaticMeshEditorViewportClient.cpp
--- a/Engine/Source/Editor/StaticMeshEditor/Private/StaticMeshEditorViewportClient.cpp
+++ b/Engine/Source/Editor/StaticMeshEditor/Private/StaticMeshEditorViewportClient.cpp
@@ -792,7 +792,7 @@ void FStaticMeshEditorViewportClient::DrawCanvas( FViewport& InViewport, FSceneV
 	{
 		int32 LOD = StaticMeshEditor->GetCurrentLODLevel();
 		return (LOD == 0) ?
-			ComputeStaticMeshLOD(StaticMesh->GetRenderData(), StaticMeshComponent->Bounds.Origin, StaticMeshComponent->Bounds.SphereRadius, View, StaticMesh->GetMinLOD().Default)
+			ComputeStaticMeshLOD(StaticMesh->GetRenderData(), StaticMeshComponent->Bounds.Origin, StaticMeshComponent->Bounds.SphereRadius, View, StaticMesh->GetDefaultMinLOD())
 			:
 			LOD - 1;
 	}();
diff --git a/Engine/Source/Editor/UnrealEd/Private/Commandlets/ContentCommandlets.cpp b/Engine/Source/Editor/UnrealEd/Private/Commandlets/ContentCommandlets.cpp
--- a/Engine/Source/Editor/UnrealEd/Private/Commandlets/ContentCommandlets.cpp
+++ b/Engine/Source/Editor/UnrealEd/Private/Commandlets/ContentCommandlets.cpp
@@ -3454,7 +3454,7 @@ int32 UStaticMeshMinLodCommandlet::Main(const FString& Params)
 							}
 						}
 					}
-					//StaticMesh->SetQualityLevelMinLOD(MoveTemp(QualityLevelMinLOD));
+					StaticMesh->SetQualityLevelMinLOD(MoveTemp(QualityLevelMinLOD));
 				}
 
 				//generate a unique group name to collect stats
diff --git a/Engine/Source/Editor/UnrealEd/Private/Fbx/FbxStaticMeshImport.cpp b/Engine/Source/Editor/UnrealEd/Private/Fbx/FbxStaticMeshImport.cpp
--- a/Engine/Source/Editor/UnrealEd/Private/Fbx/FbxStaticMeshImport.cpp
+++ b/Engine/Source/Editor/UnrealEd/Private/Fbx/FbxStaticMeshImport.cpp
@@ -2243,7 +2243,7 @@ void UnFbx::FFbxImporter::PostImportStaticMesh(UStaticMesh* StaticMesh, TArray<F
 		//Set the minimum LOD
 		if (ImportOptions->MinimumLodNumber > 0)
 		{
-			StaticMesh->SetMinLOD(ImportOptions->MinimumLodNumber);
+			StaticMesh->SetMinLODIdx(ImportOptions->MinimumLodNumber);
 		}
 
 		//User specify a number of LOD.
diff --git a/Engine/Source/Runtime/Engine/Classes/Engine/Engine.h b/Engine/Source/Runtime/Engine/Classes/Engine/Engine.h
--- a/Engine/Source/Runtime/Engine/Classes/Engine/Engine.h
+++ b/Engine/Source/Runtime/Engine/Classes/Engine/Engine.h
@@ -1419,6 +1419,9 @@ public:
 	UPROPERTY(config, EditAnywhere, Category = Screenshots)
 	FDirectoryPath GameScreenshotSaveDirectory;
 
+	UPROPERTY(config, EditAnywhere, Category = PerQualityLevelProperty, AdvancedDisplay)
+	bool UsePerQualityLevelProperty;
+
 	/** The state of the current map transition.  */
 	UPROPERTY()
 	ETransitionType TransitionType;
diff --git a/Engine/Source/Runtime/Engine/Classes/Engine/StaticMesh.h b/Engine/Source/Runtime/Engine/Classes/Engine/StaticMesh.h
--- a/Engine/Source/Runtime/Engine/Classes/Engine/StaticMesh.h
+++ b/Engine/Source/Runtime/Engine/Classes/Engine/StaticMesh.h
@@ -24,7 +24,7 @@
 #include "PerPlatformProperties.h"
 #include "RenderAssetUpdate.h"
 #include "MeshTypes.h"
-
+#include "PerQualityLevelProperties.h"
 #include "StaticMesh.generated.h"
 
 class FSpeedTreeWind;
@@ -723,6 +723,55 @@ public:
 
 #endif // #if WITH_EDITORONLY_DATA
 
+	/** Check the QualitLevel property is enabled for MinLod. */
+	bool IsMinLodQualityLevelEnable() const;
+
+	UPROPERTY()
+	/*PerQuality override. Note: Enable PerQuality override in the Project Settings/ General Settings/ UsePerQualityLevelProperty*/
+	/* Allow more flexibility to set various values driven by the Scalability or Device Profile.*/
+	FPerQualityLevelInt MinQualityLevelLOD;
+
+	const FPerQualityLevelInt& GetQualityLevelMinLOD() const
+	{
+		return MinQualityLevelLOD;
+	}
+
+	void SetQualityLevelMinLOD(FPerQualityLevelInt InMinLOD)
+	{
+		MinQualityLevelLOD = MoveTemp(InMinLOD);
+	}
+
+	UFUNCTION(BlueprintPure, Category = StaticMesh)
+	void GetMinimumLODForQualityLevels(TMap<FName, int32>& QualityLevelMinimumLODs) const
+	{
+#if WITH_EDITORONLY_DATA
+		for (const TPair<int32, int32>& Pair : GetQualityLevelMinLOD().PerQuality)
+		{
+			QualityLevelMinimumLODs.Add(QualityLevelProperty::QualityLevelToFName(Pair.Key), Pair.Value);
+		}
+#endif
+	}
+
+	UFUNCTION(BlueprintPure, Category = StaticMesh)
+	int32 GetMinimumLODForQualityLevel(const FName& QualityLevel) const
+	{
+#if WITH_EDITORONLY_DATA
+		int32 QualityLevelKey = QualityLevelProperty::FNameToQualityLevel(QualityLevel);
+		if (const int32* Result = GetQualityLevelMinLOD().PerQuality.Find(QualityLevelKey))
+		{
+			return *Result;
+		}
+#endif
+		return INDEX_NONE;
+	}
+
+	/*Choose either PerPlatform or PerQuality override. Note: Enable PerQuality override in the Project Settings/ General Settings/ UsePerQualityLevelProperty*/
+	ENGINE_API int32 GetMinLODIdx() const;
+	ENGINE_API int32 GetDefaultMinLOD() const;
+	ENGINE_API void SetMinLODIdx(int32 InMinLOD);
+
+	ENGINE_API static void OnLodStrippingQualityLevelChanged(IConsoleVariable* Variable);
+
 	/** Minimum LOD to use for rendering.  This is the default setting for the mesh and can be overridden by component settings. */
 	UE_DEPRECATED(4.27, "Please do not access this member directly; use UStaticMesh::GetMinLOD() or UStaticMesh::SetMinLOD().")
 	UPROPERTY()
diff --git a/Engine/Source/Runtime/Engine/Private/Components/StaticMeshComponent.cpp b/Engine/Source/Runtime/Engine/Private/Components/StaticMeshComponent.cpp
--- a/Engine/Source/Runtime/Engine/Private/Components/StaticMeshComponent.cpp
+++ b/Engine/Source/Runtime/Engine/Private/Components/StaticMeshComponent.cpp
@@ -2227,7 +2227,7 @@ bool UStaticMeshComponent::HasLightmapTextureCoordinates() const
 		Mesh->GetRenderData() != nullptr &&
 		Mesh->GetRenderData()->LODResources.Num() > 0)
 	{
-		int32 MeshMinLOD = Mesh->GetMinLOD().GetValue();
+		int32 MeshMinLOD = Mesh->GetMinLODIdx();
 		MeshMinLOD = FMath::Min(MeshMinLOD,  Mesh->GetRenderData()->LODResources.Num() - 1);
 		
 		return ((uint32)Mesh->GetLightMapCoordinateIndex() < Mesh->GetRenderData()->LODResources[MeshMinLOD].VertexBuffers.StaticMeshVertexBuffer.GetNumTexCoords());
diff --git a/Engine/Source/Runtime/Engine/Private/Particles/ParticleSystemRender.cpp b/Engine/Source/Runtime/Engine/Private/Particles/ParticleSystemRender.cpp
--- a/Engine/Source/Runtime/Engine/Private/Particles/ParticleSystemRender.cpp
+++ b/Engine/Source/Runtime/Engine/Private/Particles/ParticleSystemRender.cpp
@@ -1329,7 +1329,7 @@ uint32 FDynamicMeshEmitterData::GetMeshLODIndexFromProxy(const FParticleSystemSc
 	{
 		return FirstAvailableLOD;
 	}
-	const int32 EffectiveMinLOD = StaticMesh->GetMinLOD().GetValue();
+	const int32 EffectiveMinLOD = StaticMesh->GetMinLODIdx();
 	const int32 MaxLOD = StaticMesh->GetRenderData()->LODResources.Num() - 1;
 	const int32 ClampedMinLOD = FMath::Clamp(EffectiveMinLOD, FirstAvailableLOD, MaxLOD);
 	const int32 MeshLOD = (InOwnerProxy->MeshEmitterLODIndices.IsValidIndex(EmitterIndex)) ? InOwnerProxy->MeshEmitterLODIndices[EmitterIndex] : 0;
diff --git a/Engine/Source/Runtime/Engine/Private/StaticMesh.cpp b/Engine/Source/Runtime/Engine/Private/StaticMesh.cpp
--- a/Engine/Source/Runtime/Engine/Private/StaticMesh.cpp
+++ b/Engine/Source/Runtime/Engine/Private/StaticMesh.cpp
@@ -141,6 +141,23 @@ static TAutoConsoleVariable<int32> CVarStripDistanceFieldDataDuringLoad(
 	TEXT("If non-zero, data for distance fields will be discarded on load. TODO: change to discard during cook!."),
 	ECVF_ReadOnly | ECVF_RenderThreadSafe);
 
+const TCHAR* GMinLodQualityLevelCVarName = TEXT("r.StaticMesh.MinLodQualityLevel");
+const TCHAR* GMinLodQualityLevelScalabilitySection = TEXT("ViewDistanceQuality");
+int32 GMinLodQualityLevel = -1;
+static FAutoConsoleVariableRef CVarMinLodQualityLevel(
+	GMinLodQualityLevelCVarName,
+	GMinLodQualityLevel,
+	TEXT("The quality level for the Min stripping LOD. \n"),
+	FConsoleVariableDelegate::CreateStatic(&UStaticMesh::OnLodStrippingQualityLevelChanged),
+	ECVF_Scalability);
+
+int32 GDisableMinLODQualityLevel = 0;
+static FAutoConsoleVariableRef CVarCDisableMinLODQualityLevel(
+	TEXT("r.StaticMesh.DisableMinLODQualityLevel"),
+	GDisableMinLODQualityLevel,
+	TEXT("Disable MinLOD Quality Level Override (Fallback on PerPlatform). \n"),
+	ECVF_ReadOnly | ECVF_RenderThreadSafe);
+
 #if ENABLE_COOK_STATS
 namespace StaticMeshCookStats
 {
@@ -294,7 +311,34 @@ int32 FStaticMeshLODResources::GetPlatformMinLODIdx(const ITargetPlatform* Targe
 {
 #if WITH_EDITOR
 	check(TargetPlatform && StaticMesh);
-	return StaticMesh->GetMinLOD().GetValueForPlatform(*TargetPlatform->IniPlatformName());
+	if (StaticMesh->IsMinLodQualityLevelEnable())
+	{
+		// get all supported quality level from scalability + engine ini files
+		FSupportedQualityLevelArray SupportedQualityLevels = StaticMesh->GetQualityLevelMinLOD().GetSupportedQualityLevels(*TargetPlatform->GetPlatformInfo().TargetPlatformName.ToString());
+		
+		// loop through all the supported quality level to find the min lod index
+		int32 MinLodIdx = MAX_int32;
+		for (int32& QL : SupportedQualityLevels)
+		{
+			// check if have data for the supported quality level
+			if (StaticMesh->GetQualityLevelMinLOD().IsQualityLevelValid(QL))
+			{
+				MinLodIdx = FMath::Min(StaticMesh->GetQualityLevelMinLOD().GetValueForQualityLevel(QL), MinLodIdx);
+			}
+		}
+
+		if (MinLodIdx == MAX_int32)
+		{
+			MinLodIdx = StaticMesh->GetQualityLevelMinLOD().Default;
+		}
+
+		return MinLodIdx;
+	}
+	else
+	{
+		return StaticMesh->GetMinLOD().GetValueForPlatform(*TargetPlatform->IniPlatformName());
+	}
+	
 #else
 	return 0;
 #endif
@@ -2849,6 +2893,8 @@ UStaticMesh::UStaticMesh(const FObjectInitializer& ObjectInitializer)
 #if WITH_EDITOR
 	BuildCacheAutomationTestGuid.Invalidate();
 #endif
+	SetQualityLevelMinLOD(0);
+	MinQualityLevelLOD.Init(GMinLodQualityLevelCVarName, GMinLodQualityLevelScalabilitySection);
 }
 
 // We don't care if the default implementation of the destructor is cleaning up 
@@ -2957,7 +3003,7 @@ void UStaticMesh::InitResources()
 	{
 		{
 			const int32 NumLODs = GetNumLODs();
-			const int32 MinFirstLOD = GetMinLOD().GetValue();
+			const int32 MinFirstLOD = GetMinLODIdx();
 
 			CachedSRRState.NumNonStreamingLODs = GetRenderData()->NumInlinedLODs;
 			CachedSRRState.NumNonOptionalLODs = GetRenderData()->GetNumNonOptionalLODs();
@@ -3100,7 +3146,7 @@ bool UStaticMesh::HasValidRenderData(bool bCheckLODForVerts, int32 LODIndex) con
 		{
 			if (LODIndex == INDEX_NONE)
 			{
-				LODIndex = FMath::Clamp<int32>(GetMinLOD().GetValue(), 0, GetRenderData()->LODResources.Num() - 1);
+				LODIndex = FMath::Clamp<int32>(GetMinLODIdx(), 0, GetRenderData()->LODResources.Num() - 1);
 			}
 
 			return (GetRenderData()->LODResources[LODIndex].VertexBuffers.StaticMeshVertexBuffer.GetNumVertices() > 0);
@@ -5283,8 +5329,11 @@ void UStaticMesh::ExecutePostLoadInternal(FStaticMeshPostLoadContext& Context)
 	{
 		// check the MinLOD values are all within range
 		bool bFixedMinLOD = false;
+		bool bFixedQualityMinLOD = false;
 		int32 MinAvailableLOD = FMath::Max<int32>(GetRenderData()->LODResources.Num() - 1, 0);
 		FPerPlatformInt LocalMinLOD = GetMinLOD();
+		FPerQualityLevelInt QualityLocalMinLOD = GetQualityLevelMinLOD();
+
 		if (!GetRenderData()->LODResources.IsValidIndex(LocalMinLOD.Default))
 		{
 			FFormatNamedArguments Arguments;
@@ -5308,6 +5357,19 @@ void UStaticMesh::ExecutePostLoadInternal(FStaticMeshPostLoadContext& Context)
 			LocalMinLOD.Default = MinAvailableLOD;
 			bFixedMinLOD = true;
 		}
+
+		if (!GetRenderData()->LODResources.IsValidIndex(QualityLocalMinLOD.Default))
+		{
+			FFormatNamedArguments Arguments;
+			Arguments.Add(TEXT("QualityLevelMinLOD"), FText::AsNumber(QualityLocalMinLOD.Default));
+			Arguments.Add(TEXT("MinAvailLOD"), FText::AsNumber(MinAvailableLOD));
+			FMessageLog("LoadErrors").Warning()
+				->AddToken(FUObjectToken::Create(this))
+				->AddToken(FTextToken::Create(FText::Format(LOCTEXT("LoadError_BadQualityLevelMinLOD", "Min LOD value of {QualityLevelMinLOD} is out of range 0..{MinAvailLOD} and has been adjusted to {MinAvailLOD}. Please verify and resave the asset."), Arguments)));
+
+			QualityLocalMinLOD.Default = MinAvailableLOD;
+			bFixedQualityMinLOD = true;
+		}
 		
 		for (TMap<FName, int32>::TIterator It(LocalMinLOD.PerPlatform); It; ++It)
 		{
@@ -5337,12 +5399,35 @@ void UStaticMesh::ExecutePostLoadInternal(FStaticMeshPostLoadContext& Context)
 			}
 		}
 
+		for (TMap<int32, int32>::TIterator It(QualityLocalMinLOD.PerQuality); It; ++It)
+		{
+			if (!GetRenderData()->LODResources.IsValidIndex(It.Value()))
+			{
+				FFormatNamedArguments Arguments;
+				Arguments.Add(TEXT("QualityLevelMinLOD"), FText::AsNumber(It.Value()));
+				Arguments.Add(TEXT("MinAvailLOD"), FText::AsNumber(MinAvailableLOD));
+				Arguments.Add(TEXT("Quality"), FText::FromString(QualityLevelProperty::QualityLevelToFName(It.Key()).ToString()));
+				FMessageLog("LoadErrors").Warning()
+					->AddToken(FUObjectToken::Create(this))
+					->AddToken(FTextToken::Create(FText::Format(LOCTEXT("LoadError_BadMinLODOverride", "Min LOD override of {QualityLevelMinLOD} for {Quality} is out of range 0..{MinAvailLOD} and has been adjusted to {MinAvailLOD}. Please verify and resave the asset."), Arguments)));
+
+				It.Value() = MinAvailableLOD;
+				bFixedQualityMinLOD = true;
+			}
+		}
+
 		if (bFixedMinLOD)
 		{
 			SetMinLOD(MoveTemp(LocalMinLOD));
 			// Make sure Slate gets called from the game thread
 			Async(EAsyncExecution::TaskGraphMainThread, []() { FMessageLog("LoadErrors").Open(); });
 		}
+
+		if (bFixedQualityMinLOD)
+		{
+			SetQualityLevelMinLOD(MoveTemp(QualityLocalMinLOD));
+			FMessageLog("LoadErrors").Open();
+		}
 	}
 
 #endif // #if WITH_EDITOR
@@ -5895,7 +5980,7 @@ bool UStaticMesh::GetMipDataPackagePath(const int32 MipIndex, FPackagePath& OutP
 	bool bSucceed = FPackageName::DoesPackageExist(PackagePath, &PackagePath);
 	check(bSucceed);
 	OutPackagePath = MoveTemp(PackagePath);
-	OutPackageSegment = MipIndex < GetMinLOD().Default ? EPackageSegment::BulkDataOptional : EPackageSegment::BulkDataDefault;
+	OutPackageSegment = MipIndex < GetDefaultMinLOD() ? EPackageSegment::BulkDataOptional : EPackageSegment::BulkDataDefault;
 	return true;
 }
 #endif // USE_BULKDATA_STREAMING_TOKEN
@@ -7123,6 +7208,63 @@ void UStaticMesh::RemoveSocket(UStaticMeshSocket* Socket)
 	Sockets.Remove(Socket);
 }
 
+ENGINE_API int32 UStaticMesh::GetDefaultMinLOD() const
+{
+	if (IsMinLodQualityLevelEnable())
+	{
+		return GetQualityLevelMinLOD().GetDefault();
+	}
+	else
+	{
+		return GetMinLOD().GetDefault();
+	}
+}
+
+ENGINE_API int32 UStaticMesh::GetMinLODIdx() const
+{
+	if (IsMinLodQualityLevelEnable())
+	{
+		return GetQualityLevelMinLOD().GetValue(GMinLodQualityLevel);
+	}
+	else
+	{
+		return GetMinLOD().GetValue();
+	}
+}
+
+ENGINE_API void UStaticMesh::SetMinLODIdx(int32 InMinLOD)
+{
+	if (IsMinLodQualityLevelEnable())
+	{
+		SetQualityLevelMinLOD(InMinLOD);
+	}
+	else
+	{
+		SetMinLOD(InMinLOD);
+	}
+}
+
+/** Check the QualitLevel property is enabled for MinLod. */
+bool UStaticMesh::IsMinLodQualityLevelEnable() const
+{
+	return !GDisableMinLODQualityLevel && GetQualityLevelMinLOD().bIsEnabled;
+}
+
+void UStaticMesh::OnLodStrippingQualityLevelChanged(IConsoleVariable* Variable){
+#ifdef WITH_EDITOR
+	if (GEngine && GEngine->UsePerQualityLevelProperty)
+	{
+		for (TObjectIterator<UStaticMesh> It; It; ++It)
+		{
+			UStaticMesh* StaticMesh = *It;
+			if (StaticMesh && StaticMesh->IsMinLodQualityLevelEnable())
+			{
+				FStaticMeshComponentRecreateRenderStateContext Context(StaticMesh, false);
+			}
+		}
+	}
+#endif
+}
 /*-----------------------------------------------------------------------------
 UStaticMeshSocket
 -----------------------------------------------------------------------------*/
diff --git a/Engine/Source/Runtime/Engine/Private/StaticMeshRender.cpp b/Engine/Source/Runtime/Engine/Private/StaticMeshRender.cpp
--- a/Engine/Source/Runtime/Engine/Private/StaticMeshRender.cpp
+++ b/Engine/Source/Runtime/Engine/Private/StaticMeshRender.cpp
@@ -207,7 +207,7 @@ FStaticMeshSceneProxy::FStaticMeshSceneProxy(UStaticMeshComponent* InComponent,
 
 	const auto FeatureLevel = GetScene().GetFeatureLevel();
 
-	const int32 SMCurrentMinLOD = InComponent->GetStaticMesh()->GetMinLOD().GetValue();
+	const int32 SMCurrentMinLOD = InComponent->GetStaticMesh()->GetMinLODIdx();
 	int32 EffectiveMinLOD = InComponent->bOverrideMinLOD ? InComponent->MinLOD : SMCurrentMinLOD;
 
 #if WITH_EDITOR
@@ -2356,7 +2356,7 @@ FPrimitiveSceneProxy* UStaticMeshComponent::CreateSceneProxy()
 	}
 
 	const FStaticMeshLODResourcesArray& LODResources = GetStaticMesh()->GetRenderData()->LODResources;
-	if (LODResources.Num() == 0	|| LODResources[FMath::Clamp<int32>(GetStaticMesh()->GetMinLOD().Default, 0, LODResources.Num()-1)].VertexBuffers.StaticMeshVertexBuffer.GetNumVertices() == 0)
+	if (LODResources.Num() == 0	|| LODResources[FMath::Clamp<int32>(GetStaticMesh()->GetDefaultMinLOD(), 0, LODResources.Num()-1)].VertexBuffers.StaticMeshVertexBuffer.GetNumVertices() == 0)
 	{
 		UE_LOG(LogStaticMesh, Verbose, TEXT("Skipping CreateSceneProxy for StaticMeshComponent %s (LOD problems)"), *GetFullName());
 		return nullptr;
diff --git a/Engine/Source/Runtime/Engine/Public/PerPlatformProperties.h b/Engine/Source/Runtime/Engine/Public/PerPlatformProperties.h
--- a/Engine/Source/Runtime/Engine/Public/PerPlatformProperties.h
+++ b/Engine/Source/Runtime/Engine/Public/PerPlatformProperties.h
@@ -85,6 +85,12 @@ struct ENGINE_API TPerPlatformProperty
 
 #endif
 
+	_ValueType GetDefault() const
+	{
+		const _StructType* This = StaticCast<const _StructType*>(this);
+		return This->Default;
+	}
+	
 	_ValueType GetValue() const
 	{
 #if WITH_EDITORONLY_DATA && WITH_EDITOR
