Fixed missing grid in static mesh editor
- As of cl 14649186, the default mode doesn't draw the grid anymore so made the sm editor viewport client do it instead.

diff --git a/Engine/Source/Editor/StaticMeshEditor/Private/StaticMeshEditorViewportClient.cpp b/Engine/Source/Editor/StaticMeshEditor/Private/StaticMeshEditorViewportClient.cpp
--- a/Engine/Source/Editor/StaticMeshEditor/Private/StaticMeshEditorViewportClient.cpp
+++ b/Engine/Source/Editor/StaticMeshEditor/Private/StaticMeshEditorViewportClient.cpp
@@ -57,7 +57,7 @@ FStaticMeshEditorViewportClient::FStaticMeshEditorViewportClient(TWeakPtr<IStati
 	DrawHelper.bDrawPivot = false;
 	DrawHelper.bDrawWorldBox = false;
 	DrawHelper.bDrawKillZ = false;
-	DrawHelper.bDrawGrid = false;
+	DrawHelper.bDrawGrid = true;
 	DrawHelper.GridColorAxis = FColor(160,160,160);
 	DrawHelper.GridColorMajor = FColor(144,144,144);
 	DrawHelper.GridColorMinor = FColor(128,128,128);
