Remove superfluous ETabIdFlags::SaveLayout argument when creating a new FTabId.

#rnx

diff --git a/Engine/Source/Editor/StaticMeshEditor/Private/StaticMeshEditorModeUILayer.cpp b/Engine/Source/Editor/StaticMeshEditor/Private/StaticMeshEditorModeUILayer.cpp
--- a/Engine/Source/Editor/StaticMeshEditor/Private/StaticMeshEditorModeUILayer.cpp
+++ b/Engine/Source/Editor/StaticMeshEditor/Private/StaticMeshEditorModeUILayer.cpp
@@ -21,7 +21,7 @@ void UStaticMeshEditorUISubsystem::Deinitialize()
 
 void UStaticMeshEditorUISubsystem::RegisterLayoutExtensions(FLayoutExtender& Extender)
 {
-	FTabManager::FTab NewTab(FTabId(UAssetEditorUISubsystem::TopLeftTabID, ETabIdFlags::SaveLayout), ETabState::ClosedTab);
+	FTabManager::FTab NewTab(FTabId(UAssetEditorUISubsystem::TopLeftTabID), ETabState::ClosedTab);
 	Extender.ExtendLayout(FStaticMeshEditor::SocketManagerTabId, ELayoutExtensionPosition::After, NewTab);
 }
 
