+ Adding Get/Set Nanite Settings to StaticMeshEditorSubsystem
+ Adding Get/Set LODGroup to StaticMeshEditorSubsystem
* Changing FMeshNaniteSettings to BlueprintType, and supporting BlueprintReadWrite for its members

diff --git a/Engine/Source/Editor/StaticMeshEditor/Private/StaticMeshEditorSubsystem.cpp b/Engine/Source/Editor/StaticMeshEditor/Private/StaticMeshEditorSubsystem.cpp
--- a/Engine/Source/Editor/StaticMeshEditor/Private/StaticMeshEditorSubsystem.cpp
+++ b/Engine/Source/Editor/StaticMeshEditor/Private/StaticMeshEditorSubsystem.cpp
@@ -33,6 +33,7 @@
 #include "Algo/AnyOf.h"
 #include "Async/Async.h"
 #include "Misc/ScopedSlowTask.h"
+#include "Misc/FeedbackContext.h"
 
 #include "UnrealEdGlobals.h"
 #include "UnrealEd/Private/GeomFitUtils.h"
@@ -586,6 +587,41 @@ void UStaticMeshEditorSubsystem::SetLodBuildSettings(UStaticMesh* StaticMesh, co
 	}
 }
 
+FName UStaticMeshEditorSubsystem::GetLODGroup(UStaticMesh* StaticMesh)
+{
+	return StaticMesh->LODGroup;
+}
+
+bool UStaticMeshEditorSubsystem::SetLODGroup(UStaticMesh* StaticMesh, FName LODGroup, bool bRebuildImmediately)
+{
+	
+	TArray<FName> LODGroups;
+	UStaticMesh::GetLODGroups(LODGroups);
+	
+	if(LODGroups.Contains(LODGroup))
+	{
+		GWarn->BeginSlowTask(FText::Format(FText::FromString("SetLODGroup: Applying changes to %s"), FText::FromString(StaticMesh->GetName())), true, false);
+		// Close the mesh editor to prevent crashing. If changes are applied, reopen it after the mesh has been built.
+		
+		bool bStaticMeshIsEdited = false;
+		UAssetEditorSubsystem* AssetEditorSubsystem = GEditor->GetEditorSubsystem<UAssetEditorSubsystem>();
+		if (AssetEditorSubsystem->FindEditorForAsset(StaticMesh, false))
+		{
+			AssetEditorSubsystem->CloseAllEditorsForAsset(StaticMesh);
+			bStaticMeshIsEdited = true;
+		}
+		
+		StaticMesh->SetLODGroup(LODGroup, bRebuildImmediately);
+
+		GWarn->EndSlowTask();
+
+		return true;
+	}
+	UE_LOG(LogStaticMeshEditorSubsystem, Error, TEXT("SetLODGroup: %s is not a valid LODGroup"), *LODGroup.ToString());
+	return false;
+		
+}
+
 int32 UStaticMeshEditorSubsystem::ImportLOD(UStaticMesh* BaseStaticMesh, const int32 LODIndex, const FString& SourceFilename)
 {
 	TGuardValue<bool> UnattendedScriptGuard(GIsRunningUnattendedScript, true);
@@ -981,6 +1017,37 @@ TArray<float> UStaticMeshEditorSubsystem::GetLodScreenSizes(UStaticMesh* StaticM
 
 }
 
+FMeshNaniteSettings UStaticMeshEditorSubsystem::GetNaniteSettings(UStaticMesh* StaticMesh)
+{
+	return StaticMesh->NaniteSettings;
+}
+
+void UStaticMeshEditorSubsystem::SetNaniteSettings(UStaticMesh* StaticMesh, FMeshNaniteSettings NaniteSettings, bool bApplyChanges)
+{
+	// Close the mesh editor to prevent crashing. Reopen it after the mesh has been built.
+	UAssetEditorSubsystem* AssetEditorSubsystem = GEditor->GetEditorSubsystem<UAssetEditorSubsystem>();
+	bool bStaticMeshIsEdited = false;
+	if (AssetEditorSubsystem->FindEditorForAsset(StaticMesh, false))
+	{
+		AssetEditorSubsystem->CloseAllEditorsForAsset(StaticMesh);
+		bStaticMeshIsEdited = true;
+	}
+	
+	StaticMesh->NaniteSettings = NaniteSettings;
+
+	if (bApplyChanges)
+	{
+		// Request re-building of mesh with new collision shapes
+		StaticMesh->PostEditChange();
+
+		// Reopen MeshEditor on this mesh if the MeshEditor was previously opened in it
+		if (bStaticMeshIsEdited)
+		{
+			AssetEditorSubsystem->OpenEditorForAsset(StaticMesh);
+		}
+	}
+}
+
 int32 UStaticMeshEditorSubsystem::AddSimpleCollisionsWithNotification(UStaticMesh* StaticMesh, const EScriptCollisionShapeType ShapeType, bool bApplyChanges)
 {
 	TGuardValue<bool> UnattendedScriptGuard(GIsRunningUnattendedScript, true);
@@ -2281,4 +2348,5 @@ bool UStaticMeshEditorSubsystem::CreateProxyMeshActor(const TArray<class AStatic
 	return true;
 }
 
+
 #undef LOCTEXT_NAMESPACE
diff --git a/Engine/Source/Editor/StaticMeshEditor/Public/StaticMeshEditorSubsystem.h b/Engine/Source/Editor/StaticMeshEditor/Public/StaticMeshEditorSubsystem.h
--- a/Engine/Source/Editor/StaticMeshEditor/Public/StaticMeshEditorSubsystem.h
+++ b/Engine/Source/Editor/StaticMeshEditor/Public/StaticMeshEditorSubsystem.h
@@ -89,6 +89,24 @@ public:
 	UFUNCTION(BlueprintCallable, Category = "Editor Scripting | StaticMesh")
 	void SetLodBuildSettings(UStaticMesh* StaticMesh, const int32 LodIndex, const FMeshBuildSettings& BuildOptions);
 
+	/**
+	 * Get the LODGroup for the specified static mesh
+	 * @param StaticMesh
+	 * @return LODGroup
+	 */
+	UFUNCTION(BlueprintCallable,Category = "Editor Scripting | StaticMesh")
+	FName GetLODGroup(UStaticMesh* StaticMesh);
+
+	/**
+	 * Set the LODGroup for the specified static mesh
+	 * @param StaticMesh - Mesh to process.
+	 * @param LODGroup - Name of the LODGroup to apply
+	 * @param bRebuildImmediately - If true, rebuild the static mesh immediately
+	 * @return Success
+	 */
+	UFUNCTION(BlueprintCallable, Category = "Editor Scripting | StaticMesh")
+	bool SetLODGroup(UStaticMesh* StaticMesh, FName LODGroup, bool bRebuildImmediately=true);
+	
 	/**
 	 * Import or re-import a LOD into the specified base mesh. If the LOD do not exist it will import it and add it to the base static mesh. If the LOD already exist it will re-import the specified LOD.
 	 *
@@ -149,6 +167,24 @@ public:
 	UFUNCTION(BlueprintPure, Category = "Static Mesh Utilities")
 	TArray<float> GetLodScreenSizes(UStaticMesh* StaticMesh);
 
+public:
+	/**
+	* Get the Nanite Settings for the mesh
+	* @param StaticMesh        Mesh to access
+	* @return FMeshNaniteSettings struct for the given static mesh
+	*/
+	UFUNCTION(BlueprintPure, Category = "Static Mesh Utilities")
+	FMeshNaniteSettings GetNaniteSettings(UStaticMesh* StaticMesh);
+
+	/**
+	* Get the Nanite Settings for the mesh
+	* @param StaticMesh        Mesh to up[date nanite settings for
+	* @param NaniteSettings    Settings with which to update the mesh
+	* @param bApplyChanges     Indicates if changes must be applied or not.
+	*/
+	UFUNCTION(BlueprintCallable, Category = "Static Mesh Utilities")
+	void SetNaniteSettings(UStaticMesh* StaticMesh, FMeshNaniteSettings NaniteSettings, bool bApplyChanges=true);
+	
 public:
 	/**
 	 * Add simple collisions to a static mesh.
@@ -487,4 +523,5 @@ public:
 	UFUNCTION(BlueprintCallable, Category = "Editor Scripting | Dataprep")
 	bool CreateProxyMeshActor(const TArray<class AStaticMeshActor*>& ActorsToMerge, const FCreateProxyMeshActorOptions& MergeOptions, class AStaticMeshActor*& OutMergedActor);
 
+	
 };
diff --git a/Engine/Source/Runtime/Engine/Classes/Engine/EngineTypes.h b/Engine/Source/Runtime/Engine/Classes/Engine/EngineTypes.h
--- a/Engine/Source/Runtime/Engine/Classes/Engine/EngineTypes.h
+++ b/Engine/Source/Runtime/Engine/Classes/Engine/EngineTypes.h
@@ -3083,21 +3083,21 @@ struct FSkeletalMeshBuildSettings
 /**
  * Settings applied when building Nanite data.
  */
-USTRUCT()
+USTRUCT(BlueprintType)
 struct FMeshNaniteSettings
 {
 	GENERATED_BODY()
 
 	/** If true, Nanite data will be generated. */
-	UPROPERTY(EditAnywhere, Category = NaniteSettings)
+	UPROPERTY(BlueprintReadWrite, EditAnywhere, Category = NaniteSettings)
 	uint8 bEnabled : 1;
 
 	/** Position Precision. Step size is 2^(-PositionPrecision) cm. MIN_int32 is auto. */
-	UPROPERTY(EditAnywhere, Category = NaniteSettings)
+	UPROPERTY(BlueprintReadWrite, EditAnywhere, Category = NaniteSettings)
 	int32 PositionPrecision;
 
 	/** Percentage of triangles to keep from LOD0. 1.0 = no reduction, 0.0 = no triangles. */
-	UPROPERTY(EditAnywhere, Category = NaniteSettings)
+	UPROPERTY(BlueprintReadWrite, EditAnywhere, Category = NaniteSettings)
 	float PercentTriangles;
 
 	/** Default settings. */
