Fix Static mesh editor crashes if opening after closing mesh preview.

diff --git a/Engine/Source/Editor/StaticMeshEditor/Private/StaticMeshEditor.cpp b/Engine/Source/Editor/StaticMeshEditor/Private/StaticMeshEditor.cpp
--- a/Engine/Source/Editor/StaticMeshEditor/Private/StaticMeshEditor.cpp
+++ b/Engine/Source/Editor/StaticMeshEditor/Private/StaticMeshEditor.cpp
@@ -290,8 +290,11 @@ void FStaticMeshEditor::PostInitAssetEditor()
 {
 	// (Copied from FUVEditorToolkit::PostInitAssetEditor)
 	// We need the viewport client to start out focused, or else it won't get ticked until we click inside it.
-	FStaticMeshEditorViewportClient& ViewportClient = GetStaticMeshViewport()->GetViewportClient();
-	ViewportClient.ReceivedFocus(ViewportClient.Viewport);
+	if (TSharedPtr<SStaticMeshEditorViewport> StaticMeshViewport = GetStaticMeshViewport())
+	{
+		FStaticMeshEditorViewportClient& ViewportClient = StaticMeshViewport->GetViewportClient();
+		ViewportClient.ReceivedFocus(ViewportClient.Viewport);
+	}
 }
 
 void FStaticMeshEditor::GenerateSecondaryToolbar()
@@ -939,8 +942,6 @@ void FStaticMeshEditor::ExtendToolBar()
 	{
 	TSharedPtr<FExtender> ToolbarExtender = MakeShareable(new FExtender);
 
-	FStaticMeshEditorViewportClient& ViewportClient = GetStaticMeshViewport()->GetViewportClient();
-
 	FStaticMeshEditor* ThisEditor = this;
 
 	ToolbarExtender->AddToolBarExtension(
