Convert static mesh editor to use tool menus

diff --git a/Engine/Source/Editor/StaticMeshEditor/Private/StaticMeshEditor.cpp b/Engine/Source/Editor/StaticMeshEditor/Private/StaticMeshEditor.cpp
--- a/Engine/Source/Editor/StaticMeshEditor/Private/StaticMeshEditor.cpp
+++ b/Engine/Source/Editor/StaticMeshEditor/Private/StaticMeshEditor.cpp
@@ -73,6 +73,85 @@ public:
 	}
 };
 
+namespace StaticMeshEditor
+{
+	static void PopulateCollisionMenu(UToolMenu* Menu)
+	{
+		{
+			FToolMenuSection& Section = Menu->AddSection("CollisionEditCollision", LOCTEXT("CollisionEditCollisionSection", "Edit Collision"));
+			Section.AddMenuEntry(FStaticMeshEditorCommands::Get().CreateSphereCollision);
+			Section.AddMenuEntry(FStaticMeshEditorCommands::Get().CreateSphylCollision);
+			Section.AddMenuEntry(FStaticMeshEditorCommands::Get().CreateBoxCollision);
+			Section.AddMenuEntry(FStaticMeshEditorCommands::Get().CreateDOP10X);
+			Section.AddMenuEntry(FStaticMeshEditorCommands::Get().CreateDOP10Y);
+			Section.AddMenuEntry(FStaticMeshEditorCommands::Get().CreateDOP10Z);
+			Section.AddMenuEntry(FStaticMeshEditorCommands::Get().CreateDOP18);
+			Section.AddMenuEntry(FStaticMeshEditorCommands::Get().CreateDOP26);
+			Section.AddMenuEntry(FStaticMeshEditorCommands::Get().ConvertBoxesToConvex);
+			Section.AddMenuEntry(FStaticMeshEditorCommands::Get().RemoveCollision);
+
+			Section.AddMenuEntry("DeleteCollision", FGenericCommands::Get().Delete, LOCTEXT("DeleteCollision", "Delete Selected Collision"), LOCTEXT("DeleteCollisionToolTip", "Deletes the selected Collision from the mesh."));
+			Section.AddMenuEntry("DuplicateCollision", FGenericCommands::Get().Duplicate, LOCTEXT("DuplicateCollision", "Duplicate Selected Collision"), LOCTEXT("DuplicateCollisionToolTip", "Duplicates the selected Collision."));
+		}
+
+		{
+			FToolMenuSection& Section = Menu->AddSection("CollisionAutoConvexCollision");
+			Section.AddSeparator("MiscActionsSeparator");
+			Section.AddMenuEntry(FStaticMeshEditorCommands::Get().CreateAutoConvexCollision);
+		}
+
+		{
+			FToolMenuSection& Section = Menu->AddSection("CollisionCopy");
+			Section.AddMenuEntry(FStaticMeshEditorCommands::Get().CopyCollisionFromSelectedMesh);
+		}
+
+		{
+			FToolMenuSection& Section = Menu->AddSection("MeshFindSource");
+			Section.AddMenuEntry(FStaticMeshEditorCommands::Get().FindSource);
+		}
+
+		{
+			FToolMenuSection& Section = Menu->AddSection("MeshChange");
+			//Section.AddMenuEntry(FStaticMeshEditorCommands::Get().ChangeMesh);
+			Section.AddDynamicEntry("SaveGeneratedLODs", FNewToolMenuSectionDelegate::CreateLambda([](FToolMenuSection& InSection)
+			{
+				static auto* CVar = IConsoleManager::Get().FindTConsoleVariableDataInt(TEXT("r.StaticMesh.EnableSaveGeneratedLODsInPackage"));
+				if (CVar && CVar->GetValueOnGameThread() != 0)
+				{
+					InSection.AddMenuEntry(FStaticMeshEditorCommands::Get().SaveGeneratedLODs);
+				}
+			}));
+		}
+	}
+
+	static TSharedPtr<FStaticMeshEditor> GetStaticMeshEditorFromMenuContext(UAssetEditorToolkitMenuContext* InContext)
+	{
+		if (InContext)
+		{
+			if (TSharedPtr<FAssetEditorToolkit> Toolkit = InContext->Toolkit.Pin())
+			{
+				// Note: This will not detect subclasses of StaticMeshEditor
+				if (Toolkit->GetToolkitFName() == TEXT("StaticMeshEditor"))
+				{
+					return StaticCastSharedPtr<FStaticMeshEditor>(Toolkit);
+				}
+			}
+		}
+
+		return nullptr;
+	}
+
+	static TSharedPtr<FStaticMeshEditor> GetStaticMeshEditorFromMenuContext(UToolMenu* InMenu)
+	{
+		return GetStaticMeshEditorFromMenuContext(InMenu->FindContext<UAssetEditorToolkitMenuContext>());
+	}
+
+	static TSharedPtr<FStaticMeshEditor> GetStaticMeshEditorFromMenuContext(FToolMenuSection& InSection)
+	{
+		return GetStaticMeshEditorFromMenuContext(InSection.FindContext<UAssetEditorToolkitMenuContext>());
+	}
+}
+
 const FName FStaticMeshEditor::ViewportTabId( TEXT( "StaticMeshEditor_Viewport" ) );
 const FName FStaticMeshEditor::PropertiesTabId( TEXT( "StaticMeshEditor_Properties" ) );
 const FName FStaticMeshEditor::SocketManagerTabId( TEXT( "StaticMeshEditor_SocketManager" ) );
@@ -386,126 +465,55 @@ TSharedRef<IDetailCustomization> FStaticMeshEditor::MakeStaticMeshDetails()
 
 void FStaticMeshEditor::ExtendMenu()
 {
-	struct Local
 	{
-		static void FillEditMenu( FMenuBuilder& InMenuBuilder )
-		{
-			InMenuBuilder.BeginSection("Sockets", LOCTEXT("EditStaticMeshSockets", "Sockets"));
-			{
-				InMenuBuilder.AddMenuEntry( FGenericCommands::Get().Delete, "DeleteSocket", LOCTEXT("DeleteSocket", "Delete Socket"), LOCTEXT("DeleteSocketToolTip", "Deletes the selected socket from the mesh.") );
-				InMenuBuilder.AddMenuEntry( FGenericCommands::Get().Duplicate, "DuplicateSocket", LOCTEXT("DuplicateSocket", "Duplicate Socket"), LOCTEXT("DuplicateSocketToolTip", "Duplicates the selected socket.") );
-			}
-			InMenuBuilder.EndSection();
-		}
+		UToolMenu* Menu = UToolMenus::Get()->ExtendMenu("AssetEditor.StaticMeshEditor.MainMenu.Edit");
+		FToolMenuSection& Section = Menu->AddSection("Sockets", LOCTEXT("EditStaticMeshSockets", "Sockets"));
+		Section.InsertPosition = FToolMenuInsert("EditHistory", EToolMenuInsertType::After);
+		Section.AddMenuEntry("DeleteSocket", FGenericCommands::Get().Delete, LOCTEXT("DeleteSocket", "Delete Socket"), LOCTEXT("DeleteSocketToolTip", "Deletes the selected socket from the mesh."));
+		Section.AddMenuEntry("DuplicateSocket", FGenericCommands::Get().Duplicate, LOCTEXT("DuplicateSocket", "Duplicate Socket"), LOCTEXT("DuplicateSocketToolTip", "Duplicates the selected socket."));
+	}
 
-		static void FillMeshMenu( FToolMenuSection& InMenuSection)
+	{
+		if (!UToolMenus::Get()->IsMenuRegistered("StaticMeshEditor.Collision"))
 		{
-			static auto* CVar = IConsoleManager::Get().FindTConsoleVariableDataInt(TEXT("r.StaticMesh.EnableSaveGeneratedLODsInPackage"));
-			if (CVar && CVar->GetValueOnGameThread() != 0)
-			{
-				InMenuSection.AddMenuEntry(FStaticMeshEditorCommands::Get().SaveGeneratedLODs);
-			}
+			StaticMeshEditor::PopulateCollisionMenu(UToolMenus::Get()->RegisterMenu("StaticMeshEditor.Collision"));
 		}
 
-		static void FillCollisionMenu( FMenuBuilder& InMenuBuilder )
+		if (!UToolMenus::Get()->IsMenuRegistered("AssetEditor.StaticMeshEditor.MainMenu.Collision"))
 		{
-			InMenuBuilder.BeginSection("CollisionEditCollision");
-			{
-				InMenuBuilder.AddMenuEntry(FStaticMeshEditorCommands::Get().CreateSphereCollision);
-				InMenuBuilder.AddMenuEntry(FStaticMeshEditorCommands::Get().CreateSphylCollision);
-				InMenuBuilder.AddMenuEntry(FStaticMeshEditorCommands::Get().CreateBoxCollision);
-				InMenuBuilder.AddMenuEntry(FStaticMeshEditorCommands::Get().CreateDOP10X);
-				InMenuBuilder.AddMenuEntry(FStaticMeshEditorCommands::Get().CreateDOP10Y);
-				InMenuBuilder.AddMenuEntry(FStaticMeshEditorCommands::Get().CreateDOP10Z);
-				InMenuBuilder.AddMenuEntry(FStaticMeshEditorCommands::Get().CreateDOP18);
-				InMenuBuilder.AddMenuEntry(FStaticMeshEditorCommands::Get().CreateDOP26);
-				InMenuBuilder.AddMenuEntry(FStaticMeshEditorCommands::Get().ConvertBoxesToConvex);
-				InMenuBuilder.AddMenuEntry(FStaticMeshEditorCommands::Get().RemoveCollision);
-				
-				InMenuBuilder.AddMenuEntry(FGenericCommands::Get().Delete, "DeleteCollision", LOCTEXT("DeleteCollision", "Delete Selected Collision"), LOCTEXT("DeleteCollisionToolTip", "Deletes the selected Collision from the mesh."));
-				InMenuBuilder.AddMenuEntry(FGenericCommands::Get().Duplicate, "DuplicateCollision", LOCTEXT("DuplicateCollision", "Duplicate Selected Collision"), LOCTEXT("DuplicateCollisionToolTip", "Duplicates the selected Collision."));
-			}
-			InMenuBuilder.EndSection();
-
-			InMenuBuilder.BeginSection("CollisionAutoConvexCollision");
-			{
-				InMenuBuilder.AddMenuEntry(FStaticMeshEditorCommands::Get().CreateAutoConvexCollision);
-			}
-			InMenuBuilder.EndSection();
-
-			InMenuBuilder.BeginSection("CollisionCopy");
-			{
-				InMenuBuilder.AddMenuEntry(FStaticMeshEditorCommands::Get().CopyCollisionFromSelectedMesh);
-			}
-			InMenuBuilder.EndSection();
-
-			InMenuBuilder.BeginSection("MeshFindSource");
-			{
-				InMenuBuilder.AddMenuEntry(FStaticMeshEditorCommands::Get().FindSource);
-			}
-			InMenuBuilder.EndSection();
-
-			InMenuBuilder.BeginSection("MeshChange");
-			{
-				// InMenuBuilder.AddMenuEntry(FStaticMeshEditorCommands::Get().ChangeMesh);
-				static auto* CVar = IConsoleManager::Get().FindTConsoleVariableDataInt(TEXT("r.StaticMesh.EnableSaveGeneratedLODsInPackage"));
-				if (CVar && CVar->GetValueOnGameThread() != 0)
-				{
-					InMenuBuilder.AddMenuEntry(FStaticMeshEditorCommands::Get().SaveGeneratedLODs);
-				}
-			}
-			InMenuBuilder.EndSection();
-
+			UToolMenus::Get()->RegisterMenu("AssetEditor.StaticMeshEditor.MainMenu.Collision", "StaticMeshEditor.Collision");
 		}
 
-		static void GenerateCollisionMenuBar( FMenuBarBuilder& InMenuBarBuilder)
 		{
-			InMenuBarBuilder.AddPullDownMenu(
+			UToolMenu* Menu = UToolMenus::Get()->ExtendMenu("AssetEditor.StaticMeshEditor.MainMenu");
+			FToolMenuSection& Section = Menu->FindOrAddSection(NAME_None);
+			FToolMenuEntry& Entry = Section.AddSubMenu("Collision",
 				LOCTEXT("StaticMeshEditorCollisionMenu", "Collision"),
 				LOCTEXT("StaticMeshEditorCollisionMenu_ToolTip", "Opens a menu with commands for editing this mesh's collision"),
-				FNewMenuDelegate::CreateStatic(&Local::FillCollisionMenu),
-				"Collision");
+				FNewToolMenuChoice());
+			Entry.InsertPosition = FToolMenuInsert("Asset", EToolMenuInsertType::After);
 		}
-	};
-
-	TSharedPtr<FExtender> MenuExtender = MakeShareable(new FExtender);
-
-	MenuExtender->AddMenuExtension(
-		"EditHistory",
-		EExtensionHook::After,
-		GetToolkitCommands(),
-		FMenuExtensionDelegate::CreateStatic( &Local::FillEditMenu ) );
-
-
-	MenuExtender->AddMenuBarExtension(
-		"Asset",
-		EExtensionHook::After,
-		GetToolkitCommands(),
-		FMenuBarExtensionDelegate::CreateStatic( &Local::GenerateCollisionMenuBar )
-		);
-
-	AddMenuExtender(MenuExtender);
+	}
 
 	IStaticMeshEditorModule* StaticMeshEditorModule = &FModuleManager::LoadModuleChecked<IStaticMeshEditorModule>( "StaticMeshEditor" );
 	AddMenuExtender(StaticMeshEditorModule->GetMenuExtensibilityManager()->GetAllExtenders(GetToolkitCommands(), GetEditingObjects()));
 
-	UToolMenu* AssetMenu = UToolMenus::Get()->ExtendMenu("MainFrame.MainMenu.Asset");
+	UToolMenu* AssetMenu = UToolMenus::Get()->ExtendMenu("AssetEditor.StaticMeshEditor.MainMenu.Asset");
 	FToolMenuSection& AssetSection = AssetMenu->FindOrAddSection("AssetEditorActions");
-	FName StaticMeshToolkitName = GetToolkitFName();
-	FToolMenuEntry& Entry = AssetSection.AddDynamicEntry("AssetManagerEditorStaticMeshCommands", FNewToolMenuSectionDelegate::CreateLambda([StaticMeshToolkitName](FToolMenuSection& InSection)
+	FToolMenuEntry& Entry = AssetSection.AddDynamicEntry("AssetManagerEditorStaticMeshCommands", FNewToolMenuSectionDelegate::CreateLambda([](FToolMenuSection& InSection)
+	{
+		InSection.AddMenuEntry(FStaticMeshEditorCommands::Get().FindSource);
+		InSection.AddMenuEntry(FStaticMeshEditorCommands::Get().SetDrawAdditionalData);
+		InSection.AddMenuEntry(FStaticMeshEditorCommands::Get().BakeMaterials);
+		InSection.AddDynamicEntry("SaveGeneratedLODs", FNewToolMenuSectionDelegate::CreateLambda([](FToolMenuSection& InSection)
 		{
-			UAssetEditorToolkitMenuContext* MenuContext = InSection.FindContext<UAssetEditorToolkitMenuContext>();
-			if (MenuContext && MenuContext->Toolkit.IsValid() && MenuContext->Toolkit.Pin()->GetToolkitFName() == StaticMeshToolkitName)
+			static auto* CVar = IConsoleManager::Get().FindTConsoleVariableDataInt(TEXT("r.StaticMesh.EnableSaveGeneratedLODsInPackage"));
+			if (CVar && CVar->GetValueOnGameThread() != 0)
 			{
-				InSection.AddMenuEntry(FStaticMeshEditorCommands::Get().FindSource);
-				InSection.AddMenuEntry(FStaticMeshEditorCommands::Get().SetDrawAdditionalData);
-				InSection.AddMenuEntry(FStaticMeshEditorCommands::Get().BakeMaterials);
-				InSection.AddDynamicEntry("SaveGeneratedLODs", FNewToolMenuSectionDelegate::CreateStatic(&Local::FillMeshMenu));
+				InSection.AddMenuEntry(FStaticMeshEditorCommands::Get().SaveGeneratedLODs);
 			}
-		}
-	));
-
-
+		}));
+	}));
 }
 
 void FStaticMeshEditor::AddReferencedObjects( FReferenceCollector& Collector )
@@ -796,6 +804,14 @@ void FStaticMeshEditor::BindCommands()
 		Commands.SaveGeneratedLODs,
 		FExecuteAction::CreateSP(this, &FStaticMeshEditor::OnSaveGeneratedLODs));
 
+	UICommandList->MapAction(
+		Commands.ReimportMesh,
+		FExecuteAction::CreateSP(this, &FStaticMeshEditor::HandleReimportMesh));
+
+	UICommandList->MapAction(
+		Commands.ReimportAllMesh,
+		FExecuteAction::CreateSP(this, &FStaticMeshEditor::HandleReimportAllMesh));
+
 	// Collision Menu
 	UICommandList->MapAction(
 		Commands.CreateAutoConvexCollision,
@@ -820,139 +836,76 @@ void FStaticMeshEditor::BindCommands()
 		FCanExecuteAction());
 }
 
-static TSharedRef< SWidget > GenerateCollisionMenuContent(TSharedPtr<const FUICommandList> InCommandList)
+void FStaticMeshEditor::ExtendToolBar()
 {
-	FMenuBuilder MenuBuilder(true, InCommandList);
-	MenuBuilder.BeginSection("CollisionEditCollision");
+	if (!UToolMenus::Get()->IsMenuRegistered("AssetEditor.StaticMeshEditor.ToolBar.Collision"))
 	{
-		MenuBuilder.AddMenuEntry(FStaticMeshEditorCommands::Get().CreateSphereCollision);
-		MenuBuilder.AddMenuEntry(FStaticMeshEditorCommands::Get().CreateSphylCollision);
-		MenuBuilder.AddMenuEntry(FStaticMeshEditorCommands::Get().CreateBoxCollision);
-		MenuBuilder.AddMenuEntry(FStaticMeshEditorCommands::Get().CreateDOP10X);
-		MenuBuilder.AddMenuEntry(FStaticMeshEditorCommands::Get().CreateDOP10Y);
-		MenuBuilder.AddMenuEntry(FStaticMeshEditorCommands::Get().CreateDOP10Z);
-		MenuBuilder.AddMenuEntry(FStaticMeshEditorCommands::Get().CreateDOP18);
-		MenuBuilder.AddMenuEntry(FStaticMeshEditorCommands::Get().CreateDOP26);
-		MenuBuilder.AddMenuEntry(FStaticMeshEditorCommands::Get().ConvertBoxesToConvex);
-		MenuBuilder.AddMenuEntry(FStaticMeshEditorCommands::Get().RemoveCollision);
-		
-		MenuBuilder.AddMenuEntry(FGenericCommands::Get().Delete, "DeleteCollision", LOCTEXT("DeleteCollision", "Delete Selected Collision"), LOCTEXT("DeleteCollisionToolTip", "Deletes the selected Collision from the mesh."));
-		MenuBuilder.AddMenuEntry(FGenericCommands::Get().Duplicate, "DuplicateCollision", LOCTEXT("DuplicateCollision", "Duplicate Selected Collision"), LOCTEXT("DuplicateCollisionToolTip", "Duplicates the selected Collision."));
+		UToolMenus::Get()->RegisterMenu("AssetEditor.StaticMeshEditor.ToolBar.Collision", "StaticMeshEditor.Collision");
 	}
-	MenuBuilder.EndSection();
 
-	MenuBuilder.BeginSection("CollisionAutoConvexCollision");
+	// Toolbar
 	{
-		MenuBuilder.AddMenuEntry(FStaticMeshEditorCommands::Get().CreateAutoConvexCollision);
-	}
-	MenuBuilder.EndSection();
+		UToolMenu* Menu = UToolMenus::Get()->ExtendMenu("AssetEditor.StaticMeshEditor.ToolBar");
 
-	MenuBuilder.BeginSection("CollisionCopy");
-	{
-		MenuBuilder.AddMenuEntry(FStaticMeshEditorCommands::Get().CopyCollisionFromSelectedMesh);
-	}
-	MenuBuilder.EndSection();
-
-	MenuBuilder.BeginSection("MeshFindSource");
-	{
-		MenuBuilder.AddMenuEntry(FStaticMeshEditorCommands::Get().FindSource);
-	}
-	MenuBuilder.EndSection();
+		auto ConstructReimportContextMenu = [](UToolMenu* InMenu)
+		{
+			FToolMenuSection& Section = InMenu->AddSection("Reimport");
+			Section.AddMenuEntry(FStaticMeshEditorCommands::Get().ReimportMesh);
+			Section.AddMenuEntry(FStaticMeshEditorCommands::Get().ReimportAllMesh);
+		};
 
-	MenuBuilder.BeginSection("MeshChange");
-	{
-		// MenuBuilder.AddMenuEntry(FStaticMeshEditorCommands::Get().ChangeMesh);
-		static auto* CVar = IConsoleManager::Get().FindTConsoleVariableDataInt(TEXT("r.StaticMesh.EnableSaveGeneratedLODsInPackage"));
-		if (CVar && CVar->GetValueOnGameThread() != 0)
 		{
-			MenuBuilder.AddMenuEntry(FStaticMeshEditorCommands::Get().SaveGeneratedLODs);
+			FToolMenuSection& Section = Menu->AddSection("Mesh");
+			Section.InsertPosition = FToolMenuInsert("Asset", EToolMenuInsertType::After);
+
+			FToolMenuEntry& ReimportMeshEntry = Section.AddEntry(FToolMenuEntry::InitToolBarButton(FStaticMeshEditorCommands::Get().ReimportMesh));
+			ReimportMeshEntry.StyleNameOverride = "CalloutToolbar";
+			
+			FToolMenuEntry& ReimportContextMenuEntry = Section.AddEntry(FToolMenuEntry::InitComboButton(
+				"ReimportContextMenu",
+				FUIAction(),
+				FNewToolMenuDelegate::CreateLambda(ConstructReimportContextMenu),
+				TAttribute<FText>(),
+				TAttribute<FText>(),
+				TAttribute<FSlateIcon>(),
+				true
+			));
+			ReimportContextMenuEntry.StyleNameOverride = "CalloutToolbar";
 		}
-	}
-	MenuBuilder.EndSection();
 
-	return MenuBuilder.MakeWidget();
-}
-
-void FStaticMeshEditor::ExtendToolBar()
-{
-	struct Local
-	{
-		static void FillToolbar(FToolBarBuilder& ToolbarBuilder, FStaticMeshEditor* ThisEditor)
 		{
-			auto ConstructReimportContextMenu = [ThisEditor]()
-			{
-				FMenuBuilder MenuBuilder(true, nullptr);
-				MenuBuilder.AddMenuEntry(FStaticMeshEditorCommands::Get().ReimportMesh->GetLabel(),
-					FStaticMeshEditorCommands::Get().ReimportMesh->GetDescription(),
-					FSlateIcon(),
-					FUIAction(FExecuteAction::CreateSP(ThisEditor, &FStaticMeshEditor::HandleReimportMesh)));
-				MenuBuilder.AddMenuEntry(FStaticMeshEditorCommands::Get().ReimportAllMesh->GetLabel(),
-					FStaticMeshEditorCommands::Get().ReimportAllMesh->GetDescription(),
-					FSlateIcon(),
-					FUIAction(FExecuteAction::CreateSP(ThisEditor, &FStaticMeshEditor::HandleReimportAllMesh)));
-				return MenuBuilder.MakeWidget();
-			};
-
-			ToolbarBuilder.BeginSection("Mesh");
+			FToolMenuSection& Section = Menu->AddSection("Command");
+			Section.InsertPosition = FToolMenuInsert("Asset", EToolMenuInsertType::After);
+
+			FToolMenuEntry& CollisionEntry = Section.AddEntry(FToolMenuEntry::InitComboButton(
+				"Collision",
+				FUIAction(),
+				FNewToolMenuChoice(), // let registered menu be looked up by name "AssetEditor.StaticMeshEditor.ToolBar.Collision"
+				LOCTEXT("Collision_Label", "Collision"),
+				LOCTEXT("Collision_Tooltip", "Collision drawing options"),
+				FSlateIcon(FEditorStyle::GetStyleSetName(), "StaticMeshEditor.SetShowCollision")
+			));
+			CollisionEntry.StyleNameOverride = "CalloutToolbar";
+
+			Section.AddDynamicEntry("UVToolbarDynamic", FNewToolMenuSectionDelegate::CreateLambda([](FToolMenuSection& InSection)
 			{
-				ToolbarBuilder.AddToolBarButton(FUIAction(FExecuteAction::CreateSP(ThisEditor, &FStaticMeshEditor::HandleReimportMesh)),
-					NAME_None,
-					FStaticMeshEditorCommands::Get().ReimportMesh->GetLabel(),
-					FStaticMeshEditorCommands::Get().ReimportMesh->GetDescription(),
-					FStaticMeshEditorCommands::Get().ReimportMesh->GetIcon());
-				ToolbarBuilder.AddComboButton(
-					FUIAction(),
-					FOnGetContent::CreateLambda(ConstructReimportContextMenu),
-					TAttribute<FText>(),
-					TAttribute<FText>(),
-					TAttribute<FSlateIcon>(),
-					true
-				);
-			}
-			ToolbarBuilder.EndSection();
-
-			ToolbarBuilder.BeginSection("Command");
-			ToolbarBuilder.BeginStyleOverride("CalloutToolbar");
-			{
-				TSharedPtr<const FUICommandList> CommandList = ToolbarBuilder.GetTopCommandList();
-
-				ToolbarBuilder.AddComboButton(
-					FUIAction(),
-					FOnGetContent::CreateStatic(&GenerateCollisionMenuContent, CommandList),
-					LOCTEXT("Collision_Label", "Collision"),
-					LOCTEXT("Collision_Tooltip", "Collision drawing options"),
-					FSlateIcon(FEditorStyle::GetStyleSetName(), "StaticMeshEditor.SetShowCollision")
-				);
-
-				FOnGetContent OnGetUVMenuContent = FOnGetContent::CreateRaw(ThisEditor, &FStaticMeshEditor::GenerateUVChannelComboList);
-
-				ToolbarBuilder.AddComboButton(
-					FUIAction(),
-					OnGetUVMenuContent,
-					LOCTEXT("UVToolbarText", "UV"),
-					LOCTEXT("UVToolbarTooltip", "Toggles display of the static mesh's UVs for the specified channel."),
-					FSlateIcon(FEditorStyle::GetStyleSetName(), "StaticMeshEditor.SetDrawUVs"));
-			}
-			ToolbarBuilder.EndStyleOverride();
-			ToolbarBuilder.EndSection();
+				if (TSharedPtr<FStaticMeshEditor> StaticMeshEditor = StaticMeshEditor::GetStaticMeshEditorFromMenuContext(InSection))
+				{
+					FToolMenuEntry& UVToolbarEntry = InSection.AddEntry(FToolMenuEntry::InitComboButton(
+						"UVToolbar",
+						FUIAction(),
+						FNewToolMenuDelegate::CreateSP(StaticMeshEditor.ToSharedRef(), &FStaticMeshEditor::GenerateUVChannelComboList),
+						LOCTEXT("UVToolbarText", "UV"),
+						LOCTEXT("UVToolbarTooltip", "Toggles display of the static mesh's UVs for the specified channel."),
+						FSlateIcon(FEditorStyle::GetStyleSetName(), "StaticMeshEditor.SetDrawUVs")
+					));
+					UVToolbarEntry.StyleNameOverride = "CalloutToolbar";
+				}
+			}));
 		}
-	};
-
+	}
 
 	{
-	TSharedPtr<FExtender> ToolbarExtender = MakeShareable(new FExtender);
-
-	FStaticMeshEditor* ThisEditor = this;
-
-	ToolbarExtender->AddToolBarExtension(
-		"Asset",
-		EExtensionHook::After,
-		GetToolkitCommands(),
-		FToolBarExtensionDelegate::CreateStatic(&Local::FillToolbar, ThisEditor)
-	);
-
-	AddToolbarExtender(ToolbarExtender);
-
 	IStaticMeshEditorModule* StaticMeshEditorModule = &FModuleManager::LoadModuleChecked<IStaticMeshEditorModule>("StaticMeshEditor");
 
 	TArray<IStaticMeshEditorModule::FStaticMeshEditorToolbarExtender> ToolbarExtenderDelegates = StaticMeshEditorModule->GetAllStaticMeshEditorToolbarExtenders();
@@ -1501,9 +1454,8 @@ void FStaticMeshEditor::RefreshViewport()
 	}
 }
 
-TSharedRef<SWidget> FStaticMeshEditor::GenerateUVChannelComboList()
+void FStaticMeshEditor::GenerateUVChannelComboList(UToolMenu* InMenu)
 {
-	FMenuBuilder MenuBuilder(true, nullptr, EditorToolbarExtender);
 	FUIAction DrawUVsAction;
 
 	FStaticMeshEditorViewportClient& ViewportClient = GetStaticMeshViewport()->GetViewportClient();
@@ -1515,56 +1467,54 @@ TSharedRef<SWidget> FStaticMeshEditor::GenerateUVChannelComboList()
 	
 	// Add UV display functions
 	{
-		MenuBuilder.BeginSection("UVDisplayOptions");
-		MenuBuilder.AddMenuEntry(
+		FToolMenuSection& Section = InMenu->AddSection("UVDisplayOptions");
+		Section.AddMenuEntry(
+			"ShowUVSToggle",
 			LOCTEXT("ShowUVSToggle", "None"),
 			LOCTEXT("ShowUVSToggle_Tooltip", "Toggles display of the static mesh's UVs."),
 			FSlateIcon(),
 			DrawUVsAction,
-			NAME_None,
 			EUserInterfaceActionType::RadioButton
 		);
 
-		MenuBuilder.AddMenuSeparator();
+		Section.AddSeparator("ShowUVSToggleSeperator");
 		// Fill out the UV channels combo.
 		int32 MaxUVChannels = FMath::Max<int32>(GetNumUVChannels(), 1);
+		FName UVChannelIDName("UVChannel_ID");
 		for (int32 UVChannelID = 0; UVChannelID < MaxUVChannels; ++UVChannelID)
 		{
 			FUIAction MenuAction;
 			MenuAction.ExecuteAction.BindSP(this, &FStaticMeshEditor::SetCurrentViewedUVChannel, UVChannelID);
 			MenuAction.GetActionCheckState.BindSP(this, &FStaticMeshEditor::GetUVChannelCheckState, UVChannelID);
 
-			MenuBuilder.AddMenuEntry(
+			UVChannelIDName.SetNumber(UVChannelID+1);
+			Section.AddMenuEntry(
+				UVChannelIDName,
 				FText::Format(LOCTEXT("UVChannel_ID", "UV Channel {0}"), FText::AsNumber(UVChannelID)),
 				FText::Format(LOCTEXT("UVChannel_ID_ToolTip", "Overlay UV Channel {0} on the viewport"), FText::AsNumber(UVChannelID)),
 				FSlateIcon(),
 				MenuAction,
-				NAME_None,
 				EUserInterfaceActionType::RadioButton
 			);
 		}
-		MenuBuilder.EndSection();
 	}
 
 	// Add UV editing functions
 	{
-		MenuBuilder.BeginSection("UVActionOptions");
+		FToolMenuSection& Section = InMenu->AddSection("UVActionOptions");
 
 		FUIAction MenuAction;
 		MenuAction.ExecuteAction.BindSP(this, &FStaticMeshEditor::RemoveCurrentUVChannel);
 		MenuAction.CanExecuteAction.BindSP(this, &FStaticMeshEditor::CanRemoveUVChannel);
-		MenuBuilder.AddMenuEntry(
+		Section.AddMenuEntry(
+			"Remove_UVChannel",
 			LOCTEXT("Remove_UVChannel", "Remove Selected"),
 			LOCTEXT("Remove_UVChannel_ToolTip", "Remove currently selected UV channel from the static mesh"),
 			FSlateIcon(),
 			MenuAction,
-			NAME_None,
 			EUserInterfaceActionType::Button
 		);
-		MenuBuilder.EndSection();
 	}
-
-	return MenuBuilder.MakeWidget();
 }
 
 
diff --git a/Engine/Source/Editor/StaticMeshEditor/Private/StaticMeshEditor.h b/Engine/Source/Editor/StaticMeshEditor/Private/StaticMeshEditor.h
--- a/Engine/Source/Editor/StaticMeshEditor/Private/StaticMeshEditor.h
+++ b/Engine/Source/Editor/StaticMeshEditor/Private/StaticMeshEditor.h
@@ -423,7 +423,7 @@ private:
 	void OnSaveGeneratedLODs();
 
 	/** Rebuilds the UV Channel combo list and attempts to set it to the same channel. */
-	TSharedRef<SWidget> GenerateUVChannelComboList();
+	void GenerateUVChannelComboList(UToolMenu* InMenu);
 
 	/** Delete whats currently selected */
 	void DeleteSelected();
