- Notify about changes to min lod so system can respond

diff --git a/Engine/Source/Editor/StaticMeshEditor/Private/StaticMeshEditorTools.cpp b/Engine/Source/Editor/StaticMeshEditor/Private/StaticMeshEditorTools.cpp
--- a/Engine/Source/Editor/StaticMeshEditor/Private/StaticMeshEditorTools.cpp
+++ b/Engine/Source/Editor/StaticMeshEditor/Private/StaticMeshEditorTools.cpp
@@ -4448,6 +4448,7 @@ void FLevelOfDetailSettingsLayout::OnMinQualityLevelLODChanged(int32 NewValue, F
 		}
 		StaticMesh->SetQualityLevelMinLOD(MoveTemp(MinLOD));
 		StaticMesh->Modify();
+		StaticMesh->GetOnMeshChanged().Broadcast();
 	}
 	StaticMeshEditor.RefreshViewport();
 }
