'Affect Distance Field Lighting' per Static Mesh Section
* Useful to hide window frames in the distance field to work around over-occlusion

diff --git a/Engine/Source/Developer/MeshUtilities/Private/MeshRepresentationCommon.cpp b/Engine/Source/Developer/MeshUtilities/Private/MeshRepresentationCommon.cpp
--- a/Engine/Source/Developer/MeshUtilities/Private/MeshRepresentationCommon.cpp
+++ b/Engine/Source/Developer/MeshUtilities/Private/MeshRepresentationCommon.cpp
@@ -176,7 +176,7 @@ void MeshRepresentation::SetupEmbreeScene(
 					{
 						if (MaterialBlendModes.IsValidIndex(Section.MaterialIndex))
 						{
-							bTriangleIsOpaqueOrMasked = !IsTranslucentBlendMode(MaterialBlendModes[Section.MaterialIndex].BlendMode);
+							bTriangleIsOpaqueOrMasked = !IsTranslucentBlendMode(MaterialBlendModes[Section.MaterialIndex].BlendMode) && MaterialBlendModes[Section.MaterialIndex].bAffectDistanceFieldLighting;
 						}
 
 						break;
diff --git a/Engine/Source/Developer/MeshUtilities/Public/MeshUtilities.h b/Engine/Source/Developer/MeshUtilities/Public/MeshUtilities.h
--- a/Engine/Source/Developer/MeshUtilities/Public/MeshUtilities.h
+++ b/Engine/Source/Developer/MeshUtilities/Public/MeshUtilities.h
@@ -51,8 +51,9 @@ namespace ETangentOptions
 class FSignedDistanceFieldBuildMaterialData
 {
 public:
-	EBlendMode BlendMode;
-	bool bTwoSided;
+	EBlendMode BlendMode = BLEND_Opaque;
+	bool bTwoSided = false;
+	bool bAffectDistanceFieldLighting = true;
 };
 
 struct FOverlappingCorners;
diff --git a/Engine/Source/Editor/StaticMeshEditor/Private/StaticMeshEditorTools.cpp b/Engine/Source/Editor/StaticMeshEditor/Private/StaticMeshEditorTools.cpp
--- a/Engine/Source/Editor/StaticMeshEditor/Private/StaticMeshEditorTools.cpp
+++ b/Engine/Source/Editor/StaticMeshEditor/Private/StaticMeshEditorTools.cpp
@@ -2171,6 +2171,19 @@ TSharedRef<SWidget> FMeshSectionSettingsLayout::OnGenerateCustomSectionWidgetsFo
 		+ SHorizontalBox::Slot()
 		.AutoWidth()
 		.Padding(2, 0, 2, 0)
+		[
+			SNew(SCheckBox)
+			.IsChecked(this, &FMeshSectionSettingsLayout::DoesSectionAffectDistanceFieldLighting, SectionIndex)
+			.OnCheckStateChanged(this, &FMeshSectionSettingsLayout::OnSectionAffectDistanceFieldLightingChanged, SectionIndex)
+			[
+				SNew(STextBlock)
+					.Font(IDetailLayoutBuilder::GetDetailFont())
+					.Text(LOCTEXT("AffectDistanceFieldLighting", "Affect Distance Field Lighting"))
+			]
+		]
+		+ SHorizontalBox::Slot()
+		.AutoWidth()
+		.Padding(2, 0, 2, 0)
 		[
 			SNew(SCheckBox)
 			.IsChecked(this, &FMeshSectionSettingsLayout::IsSectionOpaque, SectionIndex)
@@ -2212,6 +2225,37 @@ void FMeshSectionSettingsLayout::OnSectionVisibleInRayTracingChanged(ECheckBoxSt
 	CallPostEditChange();
 }
 
+
+ECheckBoxState FMeshSectionSettingsLayout::DoesSectionAffectDistanceFieldLighting(int32 SectionIndex) const
+{
+	UStaticMesh& StaticMesh = GetStaticMesh();
+	FMeshSectionInfo Info = StaticMesh.GetSectionInfoMap().Get(LODIndex, SectionIndex);
+	return Info.bAffectDistanceFieldLighting ? ECheckBoxState::Checked : ECheckBoxState::Unchecked;
+}
+
+void FMeshSectionSettingsLayout::OnSectionAffectDistanceFieldLightingChanged(ECheckBoxState NewState, int32 SectionIndex)
+{
+	UStaticMesh& StaticMesh = GetStaticMesh();
+
+	FText TransactionTest = LOCTEXT("StaticMeshEditorSetAffectDistanceFieldLightingSectionFlag", "Staticmesh editor: Set AffectDistanceFieldLighting For section, the section will be affect distance field lighting");
+	if (NewState == ECheckBoxState::Unchecked)
+	{
+		TransactionTest = LOCTEXT("StaticMeshEditorClearAffectDistanceFieldLightingSectionFlag", "Staticmesh editor: Clear AffectDistanceFieldLighting For section");
+	}
+	FScopedTransaction Transaction(TransactionTest);
+
+	FProperty* Property = UStaticMesh::StaticClass()->FindPropertyByName(UStaticMesh::GetSectionInfoMapName());
+
+	StaticMesh.PreEditChange(Property);
+	StaticMesh.Modify();
+
+	FMeshSectionInfo Info = StaticMesh.GetSectionInfoMap().Get(LODIndex, SectionIndex);
+	Info.bAffectDistanceFieldLighting = (NewState == ECheckBoxState::Checked) ? true : false;
+	StaticMesh.GetSectionInfoMap().Set(LODIndex, SectionIndex, Info);
+	CallPostEditChange();
+}
+
+
 ECheckBoxState FMeshSectionSettingsLayout::IsSectionOpaque(int32 SectionIndex) const
 {
 	UStaticMesh& StaticMesh = GetStaticMesh();
diff --git a/Engine/Source/Editor/StaticMeshEditor/Private/StaticMeshEditorTools.h b/Engine/Source/Editor/StaticMeshEditor/Private/StaticMeshEditorTools.h
--- a/Engine/Source/Editor/StaticMeshEditor/Private/StaticMeshEditorTools.h
+++ b/Engine/Source/Editor/StaticMeshEditor/Private/StaticMeshEditorTools.h
@@ -393,6 +393,9 @@ private:
 	ECheckBoxState IsSectionVisibleInRayTracing(int32 SectionIndex) const;
 	void OnSectionVisibleInRayTracingChanged(ECheckBoxState NewState, int32 SectionIndex);
 
+	ECheckBoxState DoesSectionAffectDistanceFieldLighting(int32 SectionIndex) const;
+	void OnSectionAffectDistanceFieldLightingChanged(ECheckBoxState NewState, int32 SectionIndex);
+
 	ECheckBoxState IsSectionOpaque(int32 SectionIndex) const;
 	void OnSectionForceOpaqueFlagChanged(ECheckBoxState NewState, int32 SectionIndex);
 	
diff --git a/Engine/Source/Runtime/Core/Private/UObject/FortniteNCBranchObjectVersion.cpp b/Engine/Source/Runtime/Core/Private/UObject/FortniteNCBranchObjectVersion.cpp
--- a/Engine/Source/Runtime/Core/Private/UObject/FortniteNCBranchObjectVersion.cpp
+++ b/Engine/Source/Runtime/Core/Private/UObject/FortniteNCBranchObjectVersion.cpp
@@ -6,7 +6,7 @@ TMap<FGuid, FGuid> FFortniteNCBranchObjectVersion::GetSystemGuids()
 	TMap<FGuid, FGuid> SystemGuids;
 	const FDevSystemGuids& DevGuids = FDevSystemGuids::Get();
 
-	SystemGuids.Add(DevGuids.STATICMESH_DERIVEDDATA_VER, FGuid("F16D2728168B4CDEBFB382B28FFE9E06"));
+	SystemGuids.Add(DevGuids.STATICMESH_DERIVEDDATA_VER, FGuid("F16D2728168B4CDEBFB382B38FFE9E06"));
 
 	return SystemGuids;
 }
diff --git a/Engine/Source/Runtime/Engine/Classes/Engine/StaticMesh.h b/Engine/Source/Runtime/Engine/Classes/Engine/StaticMesh.h
--- a/Engine/Source/Runtime/Engine/Classes/Engine/StaticMesh.h
+++ b/Engine/Source/Runtime/Engine/Classes/Engine/StaticMesh.h
@@ -351,6 +351,10 @@ struct FMeshSectionInfo
 	UPROPERTY()
 	bool bVisibleInRayTracing;
 
+	/** If true, this section will affect lighting methods that use Distance Fields. */
+	UPROPERTY()
+	bool bAffectDistanceFieldLighting;
+
 	/** If true, this section will always considered opaque in ray tracing Geometry. */
 	UPROPERTY()
 	bool bForceOpaque;
@@ -361,6 +365,7 @@ struct FMeshSectionInfo
 		, bEnableCollision(true)
 		, bCastShadow(true)
 		, bVisibleInRayTracing(true)
+		, bAffectDistanceFieldLighting(true)
 		, bForceOpaque(false)
 	{
 	}
@@ -371,6 +376,7 @@ struct FMeshSectionInfo
 		, bEnableCollision(true)
 		, bCastShadow(true)
 		, bVisibleInRayTracing(true)
+		, bAffectDistanceFieldLighting(true)
 		, bForceOpaque(false)
 	{
 	}
diff --git a/Engine/Source/Runtime/Engine/Private/DistanceFieldAtlas.cpp b/Engine/Source/Runtime/Engine/Private/DistanceFieldAtlas.cpp
--- a/Engine/Source/Runtime/Engine/Private/DistanceFieldAtlas.cpp
+++ b/Engine/Source/Runtime/Engine/Private/DistanceFieldAtlas.cpp
@@ -152,6 +152,7 @@ FString BuildDistanceFieldDerivedDataKey(const FString& InMeshKey)
 void FDistanceFieldVolumeData::CacheDerivedData(const FString& InStaticMeshDerivedDataKey, const ITargetPlatform* TargetPlatform, UStaticMesh* Mesh, FStaticMeshRenderData& RenderData, UStaticMesh* GenerateSource, float DistanceFieldResolutionScale, bool bGenerateDistanceFieldAsIfTwoSided)
 {
 	FString DistanceFieldKey = BuildDistanceFieldDerivedDataKey(InStaticMeshDerivedDataKey);
+	const FMeshSectionInfoMap& SectionInfoMap = Mesh->GetSectionInfoMap();
 
 	for (int32 MaterialIndex = 0; MaterialIndex < Mesh->GetStaticMaterials().Num(); MaterialIndex++)
 	{
@@ -159,6 +160,7 @@ void FDistanceFieldVolumeData::CacheDerivedData(const FString& InStaticMeshDeriv
 		// Default material blend mode
 		MaterialData.BlendMode = BLEND_Opaque;
 		MaterialData.bTwoSided = false;
+		MaterialData.bAffectDistanceFieldLighting = SectionInfoMap.IsValidSection(0, MaterialIndex) ? SectionInfoMap.Get(0, MaterialIndex).bAffectDistanceFieldLighting : true;
 
 		UMaterialInterface* MaterialInterface = Mesh->GetStaticMaterials()[MaterialIndex].MaterialInterface;
 		if (MaterialInterface)
@@ -167,7 +169,7 @@ void FDistanceFieldVolumeData::CacheDerivedData(const FString& InStaticMeshDeriv
 			MaterialData.bTwoSided = MaterialInterface->IsTwoSided();
 		}
 
-		DistanceFieldKey += FString::Printf(TEXT("_M%u_%u"), (uint32)MaterialData.BlendMode, MaterialData.bTwoSided ? 1 : 0);
+		DistanceFieldKey += FString::Printf(TEXT("_M%u_%u_%u"), (uint32)MaterialData.BlendMode, MaterialData.bTwoSided ? 1 : 0, MaterialData.bAffectDistanceFieldLighting ? 1 : 0);
 	}
 
 	TArray<uint8> DerivedData;
@@ -203,6 +205,7 @@ void FDistanceFieldVolumeData::CacheDerivedData(const FString& InStaticMeshDeriv
 			// Default material blend mode
 			MaterialData.BlendMode = BLEND_Opaque;
 			MaterialData.bTwoSided = false;
+			MaterialData.bAffectDistanceFieldLighting = SectionInfoMap.IsValidSection(0, MaterialIndex) ? SectionInfoMap.Get(0, MaterialIndex).bAffectDistanceFieldLighting : true;
 
 			if (Mesh->GetStaticMaterials()[MaterialIndex].MaterialInterface)
 			{
diff --git a/Engine/Source/Runtime/Engine/Private/MeshCardRepresentation.cpp b/Engine/Source/Runtime/Engine/Private/MeshCardRepresentation.cpp
--- a/Engine/Source/Runtime/Engine/Private/MeshCardRepresentation.cpp
+++ b/Engine/Source/Runtime/Engine/Private/MeshCardRepresentation.cpp
@@ -138,7 +138,7 @@ FCardRepresentationAsyncQueue* GCardRepresentationAsyncQueue = NULL;
 #if WITH_EDITOR
 
 // DDC key for card representation data, must be changed when modifying the generation code or data format
-#define CARDREPRESENTATION_DERIVEDDATA_VER TEXT("69E68026-A68B-A031-A3DF-E7901648DB72")
+#define CARDREPRESENTATION_DERIVEDDATA_VER TEXT("69E68026-A68B-A032-A3DF-E7901648DB72")
 
 FString BuildCardRepresentationDerivedDataKey(const FString& InMeshKey, int32 MaxLumenMeshCards)
 {
@@ -214,12 +214,15 @@ void FCardRepresentationData::CacheDerivedData(const FString& InDDCKey, const IT
 		NewTask->MaxLumenMeshCards = MaxLumenMeshCards;
 		NewTask->bGenerateDistanceFieldAsIfTwoSided = bGenerateDistanceFieldAsIfTwoSided;
 
+		const FMeshSectionInfoMap& SectionInfoMap = Mesh->GetSectionInfoMap();
+
 		for (int32 MaterialIndex = 0; MaterialIndex < Mesh->GetStaticMaterials().Num(); MaterialIndex++)
 		{
 			FSignedDistanceFieldBuildMaterialData MaterialData;
 			// Default material blend mode
 			MaterialData.BlendMode = BLEND_Opaque;
 			MaterialData.bTwoSided = false;
+			MaterialData.bAffectDistanceFieldLighting = SectionInfoMap.IsValidSection(0, MaterialIndex) ? SectionInfoMap.Get(0, MaterialIndex).bAffectDistanceFieldLighting : true;
 
 			if (Mesh->GetStaticMaterials()[MaterialIndex].MaterialInterface)
 			{
diff --git a/Engine/Source/Runtime/Engine/Private/StaticMesh.cpp b/Engine/Source/Runtime/Engine/Private/StaticMesh.cpp
--- a/Engine/Source/Runtime/Engine/Private/StaticMesh.cpp
+++ b/Engine/Source/Runtime/Engine/Private/StaticMesh.cpp
@@ -280,6 +280,7 @@ void FStaticMeshSectionAreaWeightedTriangleSamplerBuffer::ReleaseRHI()
 
 FArchive& operator<<(FArchive& Ar, FStaticMeshSection& Section)
 {
+	// Note: this is all derived data, native versioning is not needed, but be sure to bump STATICMESH_DERIVEDDATA_VER when modifying!
 	Ar << Section.MaterialIndex;
 	Ar << Section.FirstIndex;
 	Ar << Section.NumTriangles;
@@ -308,6 +309,8 @@ FArchive& operator<<(FArchive& Ar, FStaticMeshSection& Section)
 #endif
 
 	Ar << Section.bVisibleInRayTracing;
+	Ar << Section.bAffectDistanceFieldLighting;
+
 
 	return Ar;
 }
@@ -1972,6 +1975,7 @@ void FStaticMeshRenderData::ResolveSectionInfo(UStaticMesh* Owner)
 			Section.bEnableCollision = Info.bEnableCollision;
 			Section.bCastShadow = Info.bCastShadow;
 			Section.bVisibleInRayTracing = Info.bVisibleInRayTracing;
+			Section.bAffectDistanceFieldLighting = Info.bAffectDistanceFieldLighting;
 			Section.bForceOpaque = Info.bForceOpaque;
 		}
 
diff --git a/Engine/Source/Runtime/Engine/Public/StaticMeshResources.h b/Engine/Source/Runtime/Engine/Public/StaticMeshResources.h
--- a/Engine/Source/Runtime/Engine/Public/StaticMeshResources.h
+++ b/Engine/Source/Runtime/Engine/Public/StaticMeshResources.h
@@ -196,6 +196,8 @@ struct FStaticMeshSection
 	bool bCastShadow;
 	/** If true, this section will be visible in ray tracing effects. */
 	bool bVisibleInRayTracing;
+	/** If true, this section will affect lighting methods that use Distance Fields. */
+	bool bAffectDistanceFieldLighting;
 	/** If true, this section will be considered opaque in ray tracing effects. */
 	bool bForceOpaque;
 #if WITH_EDITORONLY_DATA
@@ -216,6 +218,7 @@ struct FStaticMeshSection
 		, bEnableCollision(false)
 		, bCastShadow(true)
 		, bVisibleInRayTracing(true)
+		, bAffectDistanceFieldLighting(true)
 		, bForceOpaque(false)
 	{
 #if WITH_EDITORONLY_DATA
