UE5_MAIN:  PR #7787: Add an editor utility script function to set the LOD screen size for static meshes.

#7787

diff --git a/Engine/Source/Editor/StaticMeshEditor/Private/StaticMeshEditorSubsystem.cpp b/Engine/Source/Editor/StaticMeshEditor/Private/StaticMeshEditorSubsystem.cpp
--- a/Engine/Source/Editor/StaticMeshEditor/Private/StaticMeshEditorSubsystem.cpp
+++ b/Engine/Source/Editor/StaticMeshEditor/Private/StaticMeshEditorSubsystem.cpp
@@ -1017,6 +1017,86 @@ TArray<float> UStaticMeshEditorSubsystem::GetLodScreenSizes(UStaticMesh* StaticM
 
 }
 
+bool UStaticMeshEditorSubsystem::SetLodScreenSizes(UStaticMesh* StaticMesh, const TArray<float>& ScreenSizes)
+{
+	TGuardValue<bool> UnattendedScriptGuard(GIsRunningUnattendedScript, true);
+
+	if (!EditorScriptingHelpers::CheckIfInEditorAndPIE())
+	{
+		return false;
+	}
+
+	if (StaticMesh == nullptr)
+	{
+		UE_LOG(LogStaticMeshEditorSubsystem, Error, TEXT("SetLodScreenSizes: Input StaticMesh is null."));
+		return false;
+	}
+
+	FStaticMeshRenderData* RenderData = StaticMesh->GetRenderData();
+	if (RenderData == nullptr || (StaticMesh->GetNumLODs() == 0))
+	{
+		UE_LOG(LogStaticMeshEditorSubsystem, Error, TEXT("SetLodScreenSizes: Input StaticMesh is invalid (missing RenderData or meshes)."));
+		return false;
+	}
+
+	if (ScreenSizes.Num() == 0)
+	{
+		UE_LOG(LogStaticMeshEditorSubsystem, Error, TEXT("SetLodScreenSizes: Input ScreenSizes array is empty."));
+		return false;
+	}
+
+	// If not enough screen sizes, we set remainder to arbitrary monotonically decreasing defaults, and also ensure consecutive
+	// values are monotonically decreasing, similar to what the user interface does when editing the values manually.
+	if (ScreenSizes.Num() < StaticMesh->GetNumLODs())
+	{
+		UE_LOG(LogStaticMeshEditorSubsystem, Warning, TEXT("SetLodScreenSizes: Only %d of %d ScreenSizes provided, remainder will be set to arbitrary defaults."),
+			ScreenSizes.Num(), StaticMesh->GetNumLODs());
+	}
+
+	const float MonotonicDifference = 0.0001f;		// This difference value matches the user interface
+	bool bSanitizationRequired = false;
+
+	// Disable automatic screen size calculation, since we're providing manually overriden values
+	StaticMesh->bAutoComputeLODScreenSize = 0;
+
+	// Arbitrarily set this to a value that won't affect the monotonic clamping on the first iteration of the loop
+	float LastScreenSize = 1.0f + 2.0f * MonotonicDifference;
+
+	for (int i = 0; i < StaticMesh->GetNumLODs(); i++)
+	{
+		float ScreenSizeForLOD;
+		if (i < ScreenSizes.Num())
+		{
+			ScreenSizeForLOD = FMath::Min(ScreenSizes[i], LastScreenSize - MonotonicDifference);
+		}
+		else
+		{
+			ScreenSizeForLOD = LastScreenSize - MonotonicDifference;
+		}
+
+		ScreenSizeForLOD = FMath::Clamp(ScreenSizeForLOD, 0.0f, 1.0f);
+
+		// Track if the input values needed to be sanitized in any way, so we can warn the user this happened.
+		if (i < ScreenSizes.Num() && ScreenSizeForLOD != ScreenSizes[i])
+		{
+			bSanitizationRequired = true;
+		}
+
+		RenderData->ScreenSize[i].Default = ScreenSizeForLOD;
+		StaticMesh->GetSourceModel(i).ScreenSize = ScreenSizeForLOD;
+
+		LastScreenSize = ScreenSizeForLOD;
+	}
+
+	if (bSanitizationRequired)
+	{
+		UE_LOG(LogStaticMeshEditorSubsystem, Warning, TEXT("SetLodScreenSizes: Some input values were sanitized to be monotonic."));
+	}
+
+	return true;
+}
+
+
 FMeshNaniteSettings UStaticMeshEditorSubsystem::GetNaniteSettings(UStaticMesh* StaticMesh)
 {
 	return StaticMesh->NaniteSettings;
diff --git a/Engine/Source/Editor/StaticMeshEditor/Public/StaticMeshEditorSubsystem.h b/Engine/Source/Editor/StaticMeshEditor/Public/StaticMeshEditorSubsystem.h
--- a/Engine/Source/Editor/StaticMeshEditor/Public/StaticMeshEditorSubsystem.h
+++ b/Engine/Source/Editor/StaticMeshEditor/Public/StaticMeshEditorSubsystem.h
@@ -167,6 +167,15 @@ public:
 	UFUNCTION(BlueprintPure, Category = "Static Mesh Utilities")
 	TArray<float> GetLodScreenSizes(UStaticMesh* StaticMesh);
 
+	/**
+	 * Set LOD screen sizes.
+	 * @param	StaticMesh			Mesh to process.
+	 * @param	ScreenSizes			Array of LOD screen sizes to set.
+	 * @return A boolean indicating if setting the screen sizes was successful.
+	 */
+	UFUNCTION(BlueprintCallable, Category = "Static Mesh Utilities")
+	bool SetLodScreenSizes(UStaticMesh* StaticMesh, const TArray<float>& ScreenSizes);
+
 public:
 	/**
 	* Get the Nanite Settings for the mesh
