Fix to StaticMeshEditorTest to prevent running it with nullrhi

diff --git a/Engine/Source/Editor/StaticMeshEditor/Private/StaticMeshAutomationTests.cpp b/Engine/Source/Editor/StaticMeshEditor/Private/StaticMeshAutomationTests.cpp
--- a/Engine/Source/Editor/StaticMeshEditor/Private/StaticMeshAutomationTests.cpp
+++ b/Engine/Source/Editor/StaticMeshEditor/Private/StaticMeshAutomationTests.cpp
@@ -160,7 +160,7 @@ namespace EditorViewButtonHelper
 /**
  * Static mesh editor test
  */
-IMPLEMENT_SIMPLE_AUTOMATION_TEST(FStaticMeshEditorTest, "System.Editor.Content.Static Mesh Editor Test", EAutomationTestFlags::EditorContext | EAutomationTestFlags::EngineFilter)
+IMPLEMENT_SIMPLE_AUTOMATION_TEST(FStaticMeshEditorTest, "System.Editor.Content.Static Mesh Editor Test", EAutomationTestFlags::EditorContext | EAutomationTestFlags::EngineFilter | EAutomationTestFlags::NonNullRHI)
 
 /**
  * Take screenshots of the SME window with each of the toolbar buttons toggled separately
