Fixed UStaticMeshEditorSubSystem::SetNaniteSettings to handle nullptr for the mesh parameter and marked the mesh package dirty when we're modifying it.

diff --git a/Engine/Source/Editor/StaticMeshEditor/Private/StaticMeshEditorSubsystem.cpp b/Engine/Source/Editor/StaticMeshEditor/Private/StaticMeshEditorSubsystem.cpp
--- a/Engine/Source/Editor/StaticMeshEditor/Private/StaticMeshEditorSubsystem.cpp
+++ b/Engine/Source/Editor/StaticMeshEditor/Private/StaticMeshEditorSubsystem.cpp
@@ -1099,11 +1099,25 @@ bool UStaticMeshEditorSubsystem::SetLodScreenSizes(UStaticMesh* StaticMesh, cons
 
 FMeshNaniteSettings UStaticMeshEditorSubsystem::GetNaniteSettings(UStaticMesh* StaticMesh)
 {
-	return StaticMesh->NaniteSettings;
+	if (StaticMesh)
+	{
+		return StaticMesh->NaniteSettings;
+	}
+	else
+	{
+		FFrame::KismetExecutionMessage(TEXT("Cannot call GetNaniteSettings without a static mesh"), ELogVerbosity::Error);
+		return FMeshNaniteSettings();
+	}
 }
 
 void UStaticMeshEditorSubsystem::SetNaniteSettings(UStaticMesh* StaticMesh, FMeshNaniteSettings NaniteSettings, bool bApplyChanges)
 {
+	if (!StaticMesh)
+	{
+		FFrame::KismetExecutionMessage(TEXT("Cannot call SetNaniteSettings without a static mesh"), ELogVerbosity::Error);
+		return;
+	}
+
 	// Close the mesh editor to prevent crashing. Reopen it after the mesh has been built.
 	UAssetEditorSubsystem* AssetEditorSubsystem = GEditor->GetEditorSubsystem<UAssetEditorSubsystem>();
 	bool bStaticMeshIsEdited = false;
@@ -1113,6 +1127,7 @@ void UStaticMeshEditorSubsystem::SetNaniteSettings(UStaticMesh* StaticMesh, FMes
 		bStaticMeshIsEdited = true;
 	}
 	
+	StaticMesh->Modify();
 	StaticMesh->NaniteSettings = NaniteSettings;
 
 	if (bApplyChanges)
