Improve Nanite workflow.
Artists can import/reimport(/new file)/remove nanite hi res data from the staticmesh editor.
Artists can import/reimport(/new file)/remove nanite hi res data from the content browser context menu.
Artists can easily see in the static mesh editor if the nanite hi res data is generated or imported.

#rnx

diff --git a/Engine/Source/Developer/AssetTools/Private/AssetTypeActions/AssetTypeActions_StaticMesh.cpp b/Engine/Source/Developer/AssetTools/Private/AssetTypeActions/AssetTypeActions_StaticMesh.cpp
--- a/Engine/Source/Developer/AssetTools/Private/AssetTypeActions/AssetTypeActions_StaticMesh.cpp
+++ b/Engine/Source/Developer/AssetTools/Private/AssetTypeActions/AssetTypeActions_StaticMesh.cpp
@@ -191,13 +191,15 @@ void FAssetTypeActions_StaticMesh::GetNaniteMenu(class FMenuBuilder& MenuBuilder
 	);
 }
 
-void FAssetTypeActions_StaticMesh::GetImportLODMenu(class FMenuBuilder& MenuBuilder,TArray<TWeakObjectPtr<UStaticMesh>> Objects)
+void FAssetTypeActions_StaticMesh::GetImportLODMenu(class FMenuBuilder& MenuBuilder,TArray<TWeakObjectPtr<UStaticMesh>> Objects, bool bReimportWithNewFile)
 {
 	check(Objects.Num() > 0);
 	TWeakObjectPtr<UStaticMesh> First = Objects[0];
 	UStaticMesh* StaticMesh = First.Get();
+	//Add 1 so we can Add LOD when we are not in "Reimport With New File" submenu
+	const int32 LodCount = First->GetNumLODs() + (bReimportWithNewFile ? 0 : 1);
 
-	for(int32 LOD = 1; LOD <= First->GetNumLODs(); ++LOD)
+	for(int32 LOD = 1; LOD < LodCount; ++LOD)
 	{
 		FText LODText = FText::AsNumber( LOD );
 		FText Description = FText::Format( NSLOCTEXT("AssetTypeActions_StaticMesh", "Reimport LOD (number)", "Reimport LOD {0}"), LODText );
@@ -209,18 +211,25 @@ void FAssetTypeActions_StaticMesh::GetImportLODMenu(class FMenuBuilder& MenuBuil
 		}
 
 		MenuBuilder.AddMenuEntry( Description, ToolTip, FSlateIcon(),
-			FUIAction(FExecuteAction::CreateStatic( &FAssetTypeActions_StaticMesh::ExecuteImportMeshLOD, static_cast<UObject*>(StaticMesh), LOD) ));
+			FUIAction(FExecuteAction::CreateStatic( &FAssetTypeActions_StaticMesh::ExecuteImportMeshLOD, static_cast<UObject*>(StaticMesh), LOD, bReimportWithNewFile) ));
 	}
 }
 
 void FAssetTypeActions_StaticMesh::GetLODMenu(class FMenuBuilder& MenuBuilder, TArray<TWeakObjectPtr<UStaticMesh>> Meshes)
 {
+	constexpr bool bReimportWithNewFile = true;
 	MenuBuilder.AddSubMenu(
 		NSLOCTEXT("AssetTypeActions_StaticMesh", "StaticMesh_ImportLOD", "Import LOD"),
 		NSLOCTEXT("AssetTypeActions_StaticMesh", "StaticMesh_ImportLODtooltip", "Imports meshes into the LODs"),
-		FNewMenuDelegate::CreateSP(this, &FAssetTypeActions_StaticMesh::GetImportLODMenu, Meshes)
+		FNewMenuDelegate::CreateSP(this, &FAssetTypeActions_StaticMesh::GetImportLODMenu, Meshes, !bReimportWithNewFile)
 		);
 
+	MenuBuilder.AddSubMenu(
+		NSLOCTEXT("AssetTypeActions_StaticMesh", "StaticMesh_ImportLODWithNewFile", "Reimport LOD With New File"),
+		NSLOCTEXT("AssetTypeActions_StaticMesh", "StaticMesh_ImportLODWithNewFiletooltip", "Reimports meshes LODs with new file."),
+		FNewMenuDelegate::CreateSP(this, &FAssetTypeActions_StaticMesh::GetImportLODMenu, Meshes, bReimportWithNewFile)
+	);
+
 	MenuBuilder.AddMenuSeparator();
 
 	MenuBuilder.AddMenuEntry(
@@ -258,9 +267,9 @@ void FAssetTypeActions_StaticMesh::GetLODMenu(class FMenuBuilder& MenuBuilder, T
 	);
 }
 
-void FAssetTypeActions_StaticMesh::ExecuteImportMeshLOD(UObject* Mesh, int32 LOD)
+void FAssetTypeActions_StaticMesh::ExecuteImportMeshLOD(UObject* Mesh, int32 LOD, bool bReimportWithNewFile)
 {
-	FbxMeshUtils::ImportMeshLODDialog(Mesh, LOD);
+	FbxMeshUtils::ImportMeshLODDialog(Mesh, LOD, true, bReimportWithNewFile);
 }
 
 void FAssetTypeActions_StaticMesh::ExecuteImportHiResMesh(UStaticMesh* StaticMesh)
@@ -272,20 +281,23 @@ void FAssetTypeActions_StaticMesh::ExecuteImportHiResMesh(UStaticMesh* StaticMes
 	}
 }
 
-void FAssetTypeActions_StaticMesh::ExecuteRemoveHiResMesh(UStaticMesh* StaticMesh)
+void FAssetTypeActions_StaticMesh::ExecuteReimportHiResMeshWithNewFile(UStaticMesh* StaticMesh)
 {
-	if (StaticMesh && StaticMesh->IsHiResMeshDescriptionValid())
+	if (StaticMesh)
 	{
-		StaticMesh->Modify();
-
-		StaticMesh->ModifyHiResMeshDescription();
-		StaticMesh->ClearHiResMeshDescription();
-		StaticMesh->CommitHiResMeshDescription();
-
-		StaticMesh->GetHiResSourceModel().SourceImportFilename.Empty();
-
-		StaticMesh->PostEditChange();
+		StaticMesh->GetHiResSourceModel().SourceImportFilename = FString();
+		if (FbxMeshUtils::ImportStaticMeshHiResSourceModelDialog(StaticMesh))
+		{
+			FEditorDelegates::RefreshEditor.Broadcast();
+			FEditorSupportDelegates::RedrawAllViewports.Broadcast();
+		}
+	}
+}
 
+void FAssetTypeActions_StaticMesh::ExecuteRemoveHiResMesh(UStaticMesh* StaticMesh)
+{
+	if (FbxMeshUtils::RemoveStaticMeshHiRes(StaticMesh))
+	{
 		FEditorDelegates::RefreshEditor.Broadcast();
 		FEditorSupportDelegates::RedrawAllViewports.Broadcast();
 	}
@@ -304,16 +316,25 @@ void FAssetTypeActions_StaticMesh::GetImportHiResMenu(class FMenuBuilder& MenuBu
 
 	// Import / Reimport
 	{
-		FText Description = NSLOCTEXT("AssetTypeActions_StaticMesh", "Reimport_HighRes", "Reimport High Res");
-		FText ToolTip = NSLOCTEXT("AssetTypeActions_StaticMesh", "Reimport_HighRes_Tip", "Reimport over existing the existing high res mesh.");
 		if (!StaticMesh->IsHiResMeshDescriptionValid())
 		{
-			Description = NSLOCTEXT("AssetTypeActions_StaticMesh", "Import_HighRes", "Import High Res...");
-			ToolTip = NSLOCTEXT("AssetTypeActions_StaticMesh", "Import_HighResTip", "Adds a high resolution version of the mesh. Used instead of LOD0 by systems that can benefit from more detailed geometry (Nanite, modeling, etc.).");
+			FText Description = NSLOCTEXT("AssetTypeActions_StaticMesh", "Import_HighRes", "Import High Res...");
+			FText ToolTip = NSLOCTEXT("AssetTypeActions_StaticMesh", "Import_HighResTip", "Adds a high resolution version of the mesh. Used instead of LOD0 by systems that can benefit from more detailed geometry (Nanite, modeling, etc.).");
+			MenuBuilder.AddMenuEntry(Description, ToolTip, FSlateIcon(),
+				FUIAction(FExecuteAction::CreateStatic(&FAssetTypeActions_StaticMesh::ExecuteImportHiResMesh, StaticMesh)));
+		}
+		else
+		{
+			FText Description = NSLOCTEXT("AssetTypeActions_StaticMesh", "Reimport_HighRes", "Reimport High Res");
+			FText ToolTip = NSLOCTEXT("AssetTypeActions_StaticMesh", "Reimport_HighRes_Tip", "Reimport over existing the existing high res mesh.");
+			MenuBuilder.AddMenuEntry(Description, ToolTip, FSlateIcon(),
+				FUIAction(FExecuteAction::CreateStatic(&FAssetTypeActions_StaticMesh::ExecuteImportHiResMesh, StaticMesh)));
+
+			Description = NSLOCTEXT("AssetTypeActions_StaticMesh", "Reimport_HighRes_WithNewFile", "Reimport High Res with new file...");
+			ToolTip = NSLOCTEXT("AssetTypeActions_StaticMesh", "Reimport_HighRes_WithNewFileTip", "Reimport the high resolution version of the mesh with a new file.");
+			MenuBuilder.AddMenuEntry(Description, ToolTip, FSlateIcon(),
+				FUIAction(FExecuteAction::CreateStatic(&FAssetTypeActions_StaticMesh::ExecuteReimportHiResMeshWithNewFile, StaticMesh)));
 		}
-
-		MenuBuilder.AddMenuEntry(Description, ToolTip, FSlateIcon(),
-			FUIAction(FExecuteAction::CreateStatic(&FAssetTypeActions_StaticMesh::ExecuteImportHiResMesh, StaticMesh)));
 	}
 
 	// Remove
diff --git a/Engine/Source/Developer/AssetTools/Private/AssetTypeActions/AssetTypeActions_StaticMesh.h b/Engine/Source/Developer/AssetTools/Private/AssetTypeActions/AssetTypeActions_StaticMesh.h
--- a/Engine/Source/Developer/AssetTools/Private/AssetTypeActions/AssetTypeActions_StaticMesh.h
+++ b/Engine/Source/Developer/AssetTools/Private/AssetTypeActions/AssetTypeActions_StaticMesh.h
@@ -59,13 +59,13 @@ private:
 	void GetNaniteMenu(class FMenuBuilder& MenuBuilder, TArray<TWeakObjectPtr<UStaticMesh>> Meshes);
 
 	/** Handler to provide the list of LODs that can be imported or reimported */
-	void GetImportLODMenu(class FMenuBuilder& MenuBuilder,TArray<TWeakObjectPtr<UStaticMesh>> Objects);
+	void GetImportLODMenu(class FMenuBuilder& MenuBuilder,TArray<TWeakObjectPtr<UStaticMesh>> Objects, const bool bWithNewFile);
 
 	/** Handler to provide the LOD sub-menu. Hides away LOD actions - includes Import LOD sub menu */
 	void GetLODMenu(class FMenuBuilder& MenuBuilder, TArray<TWeakObjectPtr<UStaticMesh>> Meshes);
 
 	/** Handler for calling import methods */
-	static void ExecuteImportMeshLOD(UObject* Mesh, int32 LOD);
+	static void ExecuteImportMeshLOD(UObject* Mesh, int32 LOD, bool bReimportWithNewFile);
 
 	/** Builds the High Res mesh sub-menu as part of the LOD menu */
 	void GetImportHiResMenu(class FMenuBuilder& MenuBuilder, TArray<TWeakObjectPtr<UStaticMesh>> Objects);
@@ -73,6 +73,9 @@ private:
 	/** Handles import and reimport of the high res source model */
 	static void ExecuteImportHiResMesh(UStaticMesh* Mesh);
 
+	/** Handles reimport of the high res source model with a new specified file. */
+	static void ExecuteReimportHiResMeshWithNewFile(UStaticMesh* StaticMesh);
+
 	/** Handles removing the high res source model from the mesh */
 	static void ExecuteRemoveHiResMesh(UStaticMesh* Mesh);
 
diff --git a/Engine/Source/Editor/StaticMeshEditor/Private/StaticMeshEditorTools.cpp b/Engine/Source/Editor/StaticMeshEditor/Private/StaticMeshEditorTools.cpp
--- a/Engine/Source/Editor/StaticMeshEditor/Private/StaticMeshEditorTools.cpp
+++ b/Engine/Source/Editor/StaticMeshEditor/Private/StaticMeshEditorTools.cpp
@@ -5066,10 +5066,57 @@ void FNaniteSettingsLayout::UpdateSettings(const FMeshNaniteSettings& InSettings
 void FNaniteSettingsLayout::AddToDetailsPanel(IDetailLayoutBuilder& DetailBuilder)
 {
 	UStaticMesh* StaticMesh = StaticMeshEditor.GetStaticMesh();
+	TWeakObjectPtr<UStaticMesh> WeakStaticMesh = StaticMesh;
+
+	const FText NaniteCategoryName = LOCTEXT("NaniteSettingsCategory", "Nanite Settings");
+
+	auto NaniteCategoryContentLambda = [WeakStaticMesh]()
+	{
+		UStaticMesh* StaticMesh = WeakStaticMesh.Get();
+		if (StaticMesh && StaticMesh->NaniteSettings.bEnabled)
+		{
+			if (!StaticMesh->GetHiResSourceModel().SourceImportFilename.IsEmpty())
+			{
+				return LOCTEXT("NaniteSettingsCategory_Imported", "[Imported]");
+			}
+		}
+		return FText::GetEmpty();
+	};
+
+	auto NaniteCategoryContentTooltipLambda = [WeakStaticMesh]()
+	{
+		UStaticMesh* StaticMesh = WeakStaticMesh.Get();
+		if (StaticMesh && StaticMesh->NaniteSettings.bEnabled)
+		{
+			if (!StaticMesh->GetHiResSourceModel().SourceImportFilename.IsEmpty())
+			{
+				return FText::Format(LOCTEXT("NaniteSettingsCategory_ImportedTooltip", "The nanite high resolution data is imported from file {0}"), FText::FromString(StaticMesh->GetHiResSourceModel().SourceImportFilename));
+			}
+		}
+		return FText::GetEmpty();
+	};
 
 	IDetailCategoryBuilder& NaniteSettingsCategory =
-		DetailBuilder.EditCategory("NaniteSettings", LOCTEXT("NaniteSettingsCategory", "Nanite Settings"));
+		DetailBuilder.EditCategory("NaniteSettings", NaniteCategoryName);
 	NaniteSettingsCategory.SetSortOrder(10);
+	
+	NaniteSettingsCategory.HeaderContent
+	(
+		SNew( SHorizontalBox )
+		+ SHorizontalBox::Slot()
+		.AutoWidth()
+		[
+			SNew(SBox)
+			.Padding(FMargin(5.0f, 0.0f))
+			[
+				SNew(STextBlock)
+				.Text_Lambda(NaniteCategoryContentLambda)
+				.ToolTipText_Lambda(NaniteCategoryContentTooltipLambda)
+				.Font(IDetailLayoutBuilder::GetDetailFontItalic())
+					
+			]
+		]
+	);
 
 	TSharedPtr<SCheckBox> NaniteEnabledCheck;
 	{
@@ -5221,22 +5268,108 @@ void FNaniteSettingsLayout::AddToDetailsPanel(IDetailLayoutBuilder& DetailBuilde
 		];
 	}
 
-	NaniteSettingsCategory.AddCustomRow(LOCTEXT("ApplyChanges", "Apply Changes"))
-	.ValueContent()
-	.HAlign(HAlign_Left)
-	.VAlign(VAlign_Center)
-	[
-		SNew(SButton)
-		.OnClicked(this, &FNaniteSettingsLayout::OnApply)
-		.IsEnabled(this, &FNaniteSettingsLayout::IsApplyNeeded)
-		.VAlign(VAlign_Center)
-		.HAlign(HAlign_Center)
+	//Source import filename
+	{
+		FString FileFilterText = TEXT("Filmbox (*.fbx)|*.fbx|All files (*.*)|*.*");
+		NaniteSettingsCategory.AddCustomRow( LOCTEXT("NANITE_SourceImportFilename", "Source Import Filename") )
+
+		.IsEnabled(TAttribute<bool>::Create(TAttribute<bool>::FGetter::CreateLambda([NaniteEnabledCheck]() -> bool {return NaniteEnabledCheck->IsChecked(); } )))
+		.NameContent()
 		[
 			SNew(STextBlock)
-			.Text(LOCTEXT("ApplyChanges", "Apply Changes"))
-			.Font(DetailBuilder.GetDetailFont())
+			.Font(IDetailLayoutBuilder::GetDetailFont())
+			.Text(LOCTEXT("NANITE_SourceImportFilename", "Source Import Filename"))
+			.ToolTipText(LOCTEXT("NANITE_SourceImportFilenameTooltip", "The file path that was used to import this hi res nanite mesh."))
 		]
-	];
+		.ValueContent()
+		.VAlign(VAlign_Center)
+		[
+			SNew(SFilePathPicker)
+			.BrowseButtonImage(FAppStyle::GetBrush("PropertyWindow.Button_Ellipsis"))
+			.BrowseButtonStyle(FAppStyle::Get(), "HoverHintOnly")
+			.BrowseButtonToolTip(LOCTEXT("NaniteSourceImportFilenamePathLabel_Tooltip", "Choose a nanite hi res source import file"))
+			.BrowseDirectory(FEditorDirectories::Get().GetLastDirectory(ELastDirectory::GENERIC_OPEN))
+			.BrowseTitle(LOCTEXT("NaniteSourceImportFilenameBrowseTitle", "Nanite hi res source import file picker..."))
+			.FilePath(this, &FNaniteSettingsLayout::GetHiResSourceFilename)
+			.FileTypeFilter(FileFilterText)
+			.OnPathPicked(this, &FNaniteSettingsLayout::SetHiResSourceFilename)
+		];
+	}
+
+	//Nanite import button
+	{
+		NaniteSettingsCategory.AddCustomRow(LOCTEXT("NaniteHiResButtons", "Nanite Hi Res buttons"))
+		.ValueContent()
+		.HAlign(HAlign_Fill)
+		[
+			SNew(SUniformWrapPanel)
+			+ SUniformWrapPanel::Slot() // Nanite apply changes
+			[
+				SNew(SButton)
+				.OnClicked(this, &FNaniteSettingsLayout::OnApply)
+				.IsEnabled(this, &FNaniteSettingsLayout::IsApplyNeeded)
+				.VAlign(VAlign_Center)
+				.HAlign(HAlign_Center)
+				[
+					SNew(STextBlock)
+					.Text(LOCTEXT("ApplyChanges", "Apply Changes"))
+					.Font(DetailBuilder.GetDetailFont())
+				]
+			]
+			+ SUniformWrapPanel::Slot() //Nanite import button
+			[
+				SNew(SButton)
+				.OnClicked(this, &FNaniteSettingsLayout::OnImportHiRes)
+				.IsEnabled(this, &FNaniteSettingsLayout::IsHiResDataEmpty)
+				.VAlign(VAlign_Center)
+				.HAlign(HAlign_Center)
+				[
+					SNew(STextBlock)
+					.Text(LOCTEXT("NaniteImportHiRes", "Import"))
+					.Font(DetailBuilder.GetDetailFont())
+				]
+			]
+			+ SUniformWrapPanel::Slot() //Nanite remove button
+			[
+				SNew(SButton)
+				.OnClicked(this, &FNaniteSettingsLayout::OnRemoveHiRes)
+				.IsEnabled(this, &FNaniteSettingsLayout::DoesHiResDataExists)
+				.VAlign(VAlign_Center)
+				.HAlign(HAlign_Center)
+				[
+					SNew(STextBlock)
+					.Text(LOCTEXT("NaniteRemoveHiRes", "Remove"))
+					.Font(DetailBuilder.GetDetailFont())
+				]
+			]
+			+ SUniformWrapPanel::Slot() //Nanite reimport button
+			[
+				SNew(SButton)
+				.OnClicked(this, &FNaniteSettingsLayout::OnReimportHiRes)
+				.IsEnabled(this, &FNaniteSettingsLayout::DoesHiResDataExists)
+				.VAlign(VAlign_Center)
+				.HAlign(HAlign_Center)
+				[
+					SNew(STextBlock)
+					.Text(LOCTEXT("NaniteReimportHiRes", "Reimport"))
+					.Font(DetailBuilder.GetDetailFont())
+				]
+			]
+			+ SUniformWrapPanel::Slot() //Nanite reimport with new file button
+			[
+				SNew(SButton)
+				.OnClicked(this, &FNaniteSettingsLayout::OnReimportHiResWithNewFile)
+				.IsEnabled(this, &FNaniteSettingsLayout::DoesHiResDataExists)
+				.VAlign(VAlign_Center)
+				.HAlign(HAlign_Center)
+				[
+					SNew(STextBlock)
+					.Text(LOCTEXT("NaniteReimportHiResWithNewFile", "Reimport New File"))
+					.Font(DetailBuilder.GetDetailFont())
+				]
+			]
+		];
+	}
 }
 
 bool FNaniteSettingsLayout::IsApplyNeeded() const
@@ -5457,4 +5590,100 @@ void FNaniteSettingsLayout::OnFallbackRelativeErrorChanged(float NewValue)
 	NaniteSettings.FallbackRelativeError = NewValue;
 }
 
+FString FNaniteSettingsLayout::GetHiResSourceFilename() const
+{
+	if (UStaticMesh* StaticMesh = StaticMeshEditor.GetStaticMesh())
+	{
+		return StaticMesh->GetHiResSourceModel().SourceImportFilename;
+	}
+	return FString();
+}
+
+void FNaniteSettingsLayout::SetHiResSourceFilename(const FString& NewSourceFile)
+{
+	//Reimport with new file
+	UStaticMesh* StaticMesh = StaticMeshEditor.GetStaticMesh();
+	if(!StaticMesh)
+	{
+		return;
+	}
+	
+	if (StaticMesh->GetHiResSourceModel().SourceImportFilename.Equals(NewSourceFile))
+	{
+		return;
+	}
+
+	StaticMesh->GetHiResSourceModel().SourceImportFilename = NewSourceFile;
+	//Trig a reimport with new file
+	FbxMeshUtils::ImportStaticMeshHiResSourceModelDialog(StaticMesh);
+	StaticMeshEditor.RefreshTool();
+}
+
+
+bool FNaniteSettingsLayout::DoesHiResDataExists() const
+{
+	UStaticMesh* StaticMesh = StaticMeshEditor.GetStaticMesh();
+	if (!StaticMesh)
+	{
+		return false;
+	}
+
+	return (StaticMesh->GetHiResMeshDescription() != nullptr);
+}
+
+bool FNaniteSettingsLayout::IsHiResDataEmpty() const
+{
+	UStaticMesh* StaticMesh = StaticMeshEditor.GetStaticMesh();
+	if (!StaticMesh)
+	{
+		return true;
+	}
+
+	return (StaticMesh->GetHiResMeshDescription() == nullptr);
+}
+
+FReply FNaniteSettingsLayout::OnImportHiRes()
+{
+	
+	if (UStaticMesh* StaticMesh = StaticMeshEditor.GetStaticMesh())
+	{
+		StaticMesh->GetHiResSourceModel().SourceImportFilename = FString();
+		FbxMeshUtils::ImportStaticMeshHiResSourceModelDialog(StaticMesh);
+		StaticMeshEditor.RefreshTool();
+	}
+	return FReply::Handled();
+}
+
+FReply FNaniteSettingsLayout::OnRemoveHiRes()
+{
+	if (UStaticMesh* StaticMesh = StaticMeshEditor.GetStaticMesh())
+	{
+		StaticMesh->GetHiResSourceModel().SourceImportFilename = FString();
+		FbxMeshUtils::RemoveStaticMeshHiRes(StaticMesh);
+		StaticMeshEditor.RefreshTool();
+	}
+	return FReply::Handled();
+}
+
+FReply FNaniteSettingsLayout::OnReimportHiRes()
+{
+	if (UStaticMesh* StaticMesh = StaticMeshEditor.GetStaticMesh())
+	{
+		FbxMeshUtils::ImportStaticMeshHiResSourceModelDialog(StaticMesh);
+		StaticMeshEditor.RefreshTool();
+	}
+	return FReply::Handled();
+}
+
+FReply FNaniteSettingsLayout::OnReimportHiResWithNewFile()
+{
+	if (UStaticMesh* StaticMesh = StaticMeshEditor.GetStaticMesh())
+	{
+		StaticMesh->GetHiResSourceModel().SourceImportFilename = FString();
+		FbxMeshUtils::ImportStaticMeshHiResSourceModelDialog(StaticMesh);
+		StaticMeshEditor.RefreshTool();
+	}
+	return FReply::Handled();
+}
+
 #undef LOCTEXT_NAMESPACE
\ No newline at end of file
diff --git a/Engine/Source/Editor/StaticMeshEditor/Private/StaticMeshEditorTools.h b/Engine/Source/Editor/StaticMeshEditor/Private/StaticMeshEditorTools.h
--- a/Engine/Source/Editor/StaticMeshEditor/Private/StaticMeshEditorTools.h
+++ b/Engine/Source/Editor/StaticMeshEditor/Private/StaticMeshEditorTools.h
@@ -730,6 +730,17 @@ private:
 	float GetFallbackRelativeError() const;
 	void OnFallbackRelativeErrorChanged(float NewValue);
 
+	FString GetHiResSourceFilename() const;
+	void SetHiResSourceFilename(const FString& NewSourceFile);
+
+	bool DoesHiResDataExists() const;
+	bool IsHiResDataEmpty() const;
+	
+	FReply OnImportHiRes();
+	FReply OnRemoveHiRes();
+	FReply OnReimportHiRes();
+	FReply OnReimportHiResWithNewFile();
+
 private:
 	/** The Static Mesh Editor this tool is associated with. */
 	FStaticMeshEditor& StaticMeshEditor;
diff --git a/Engine/Source/Editor/UnrealEd/Private/FbxMeshUtils.cpp b/Engine/Source/Editor/UnrealEd/Private/FbxMeshUtils.cpp
--- a/Engine/Source/Editor/UnrealEd/Private/FbxMeshUtils.cpp
+++ b/Engine/Source/Editor/UnrealEd/Private/FbxMeshUtils.cpp
@@ -815,7 +815,7 @@ namespace FbxMeshUtils
 		return ChosenFilname;
 	}
 
-	bool ImportMeshLODDialog(class UObject* SelectedMesh, int32 LODLevel, bool bNotifyCB /*= true*/)
+	bool ImportMeshLODDialog(class UObject* SelectedMesh, int32 LODLevel, bool bNotifyCB /*= true*/, bool bReimportWithNewFile /*= false*/)
 	{
 		if(!SelectedMesh)
 		{
@@ -829,7 +829,7 @@ namespace FbxMeshUtils
 		UInterchangeAssetImportData* SelectedInterchangeAssetImportData = nullptr;
 		if (SkeletalMesh)
 		{
-			if (SkeletalMesh->IsValidLODIndex(LODLevel))
+			if (!bReimportWithNewFile && SkeletalMesh->IsValidLODIndex(LODLevel))
 			{
 				FilenameToImport = SkeletalMesh->GetLODInfo(LODLevel)->SourceImportFilename.IsEmpty() ?
 					SkeletalMesh->GetLODInfo(LODLevel)->SourceImportFilename :
@@ -839,7 +839,7 @@ namespace FbxMeshUtils
 		}
 		else if (StaticMesh)
 		{
-			if (StaticMesh->IsSourceModelValid(LODLevel))
+			if (!bReimportWithNewFile && StaticMesh->IsSourceModelValid(LODLevel))
 			{
 				const FStaticMeshSourceModel& SourceModel = StaticMesh->GetSourceModel(LODLevel);
 				FilenameToImport = SourceModel.SourceImportFilename.IsEmpty() ?
@@ -1033,6 +1033,25 @@ namespace FbxMeshUtils
 		return bImportSuccess;
 	}
 
+	bool RemoveStaticMeshHiRes(UStaticMesh* StaticMesh)
+	{
+		if (!StaticMesh || !StaticMesh->IsHiResMeshDescriptionValid())
+		{
+			return false;
+		}
+
+		StaticMesh->Modify();
+
+		StaticMesh->ModifyHiResMeshDescription();
+		StaticMesh->ClearHiResMeshDescription();
+		StaticMesh->CommitHiResMeshDescription();
+
+		StaticMesh->GetHiResSourceModel().SourceImportFilename.Empty();
+
+		StaticMesh->PostEditChange();
+		return true;
+	}
+
 	void SetImportOption(UFbxImportUI* ImportUI)
 	{
 		UnFbx::FFbxImporter* FFbxImporter = UnFbx::FFbxImporter::GetInstance();
diff --git a/Engine/Source/Editor/UnrealEd/Private/ImportUtils/StaticMeshImportUtils.cpp b/Engine/Source/Editor/UnrealEd/Private/ImportUtils/StaticMeshImportUtils.cpp
--- a/Engine/Source/Editor/UnrealEd/Private/ImportUtils/StaticMeshImportUtils.cpp
+++ b/Engine/Source/Editor/UnrealEd/Private/ImportUtils/StaticMeshImportUtils.cpp
@@ -655,6 +655,32 @@ TSharedPtr<FExistingStaticMeshData> StaticMeshImportUtils::SaveExistingStaticMes
 		}
 	}
 
+	
+	/******************************************
+	 * Nanite Begin
+	 */
+
+	//Nanite Save the settings
+	ExistingMeshDataPtr->ExistingNaniteSettings = ExistingMesh->NaniteSettings;
+
+	//Nanite Save the source model
+	const FStaticMeshSourceModel& HiResSourceModel = ExistingMesh->GetHiResSourceModel();
+	ExistingMeshDataPtr->HiResSourceData.ExistingBuildSettings = HiResSourceModel.BuildSettings;
+	ExistingMeshDataPtr->HiResSourceData.ExistingReductionSettings = HiResSourceModel.ReductionSettings;
+	ExistingMeshDataPtr->HiResSourceData.ExistingScreenSize = HiResSourceModel.ScreenSize;
+	ExistingMeshDataPtr->HiResSourceData.ExistingSourceImportFilename = HiResSourceModel.SourceImportFilename;
+		
+	//Nanite Save the hi res mesh description
+	if(const FMeshDescription* HiResMeshDescription = ExistingMesh->GetHiResMeshDescription())
+	{
+		ExistingMeshDataPtr->HiResSourceData.ExistingMeshDescription = MakeUnique<FMeshDescription>(*HiResMeshDescription);
+	}
+
+	/*
+	 * Nanite End
+	 ******************************************/
+
+
 	int32 TotalMaterialIndex = ExistingMeshDataPtr->ExistingMaterials.Num();
 	for (int32 SourceModelIndex = 0; SourceModelIndex < ExistingMesh->GetNumSourceModels(); SourceModelIndex++)
 	{
@@ -707,7 +733,7 @@ TSharedPtr<FExistingStaticMeshData> StaticMeshImportUtils::SaveExistingStaticMes
 			ExistingLODData.ExistingReductionSettings.PercentVertices = 1.0f;
 			ExistingLODData.ExistingReductionSettings.MaxDeviation = 0.0f;
 		}
-		ExistingLODData.ExistingScreenSize = SourceModel.ScreenSize.Default;
+		ExistingLODData.ExistingScreenSize = SourceModel.ScreenSize;
 		ExistingLODData.ExistingSourceImportFilename = SourceModel.SourceImportFilename;
 
 		const FMeshDescription* MeshDescription = ExistingMesh->GetMeshDescription(SourceModelIndex);
@@ -750,8 +776,6 @@ TSharedPtr<FExistingStaticMeshData> StaticMeshImportUtils::SaveExistingStaticMes
 	ExistingMeshDataPtr->ExistingNeverStream = ExistingMesh->NeverStream;
 	ExistingMeshDataPtr->ExistingNumCinematicMipLevels = ExistingMesh->NumCinematicMipLevels;
 
-	ExistingMeshDataPtr->ExistingNaniteSettings = ExistingMesh->NaniteSettings;
-
 	UFbxStaticMeshImportData* ImportData = Cast<UFbxStaticMeshImportData>(ExistingMesh->AssetImportData);
 	if (ImportData && ExistingMeshDataPtr->UseMaterialNameSlotWorkflow)
 	{
@@ -1289,7 +1313,37 @@ void StaticMeshImportUtils::RestoreExistingMeshData(const TSharedPtr<const FExis
 	NewMesh->NeverStream = ExistingMeshDataPtr->ExistingNeverStream;
 	NewMesh->NumCinematicMipLevels = ExistingMeshDataPtr->ExistingNumCinematicMipLevels;
 
+	/******************************************
+	 * Nanite Begin
+	 */
+
+	 //Nanite Restore the settings
 	NewMesh->NaniteSettings = ExistingMeshDataPtr->ExistingNaniteSettings;
+
+	//Nanite Save the source model
+	FStaticMeshSourceModel& HiResSourceModel = NewMesh->GetHiResSourceModel();
+	HiResSourceModel.BuildSettings = ExistingMeshDataPtr->HiResSourceData.ExistingBuildSettings;
+	HiResSourceModel.ReductionSettings = ExistingMeshDataPtr->HiResSourceData.ExistingReductionSettings;
+	HiResSourceModel.ScreenSize = ExistingMeshDataPtr->HiResSourceData.ExistingScreenSize;
+	HiResSourceModel.SourceImportFilename = ExistingMeshDataPtr->HiResSourceData.ExistingSourceImportFilename;
+
+	//Nanite Restore the hires mesh description
+	if (ExistingMeshDataPtr->HiResSourceData.ExistingMeshDescription.IsValid())
+	{
+		FMeshDescription* HiResMeshDescription = NewMesh->GetHiResMeshDescription();
+		if (HiResMeshDescription == nullptr)
+		{
+			HiResMeshDescription = NewMesh->CreateHiResMeshDescription();
+		}
+		check(HiResMeshDescription);
+		NewMesh->ModifyHiResMeshDescription();
+		*HiResMeshDescription = MoveTemp(*ExistingMeshDataPtr->HiResSourceData.ExistingMeshDescription);
+		NewMesh->CommitHiResMeshDescription();
+	}
+
+	/*
+	 * Nanite End
+	 ******************************************/	
 }
 
 #undef LOCTEXT_NAMESPACE
diff --git a/Engine/Source/Editor/UnrealEd/Public/FbxMeshUtils.h b/Engine/Source/Editor/UnrealEd/Public/FbxMeshUtils.h
--- a/Engine/Source/Editor/UnrealEd/Public/FbxMeshUtils.h
+++ b/Engine/Source/Editor/UnrealEd/Public/FbxMeshUtils.h
@@ -48,7 +48,7 @@ namespace FbxMeshUtils
 	 * @param LODLevel			The level of the lod to import
 	 * @return					Whether or not the import succeeded
 	 */
-	UNREALED_API bool ImportMeshLODDialog( UObject* Mesh, int32 LODLevel, bool bNotifyCB = true );
+	UNREALED_API bool ImportMeshLODDialog( UObject* Mesh, int32 LODLevel, bool bNotifyCB = true, bool bReimportWithNewFile = false);
 
 	/**
 	 * Imports a mesh as the high res source model of the given mesh
@@ -59,6 +59,14 @@ namespace FbxMeshUtils
 	 */
 	UNREALED_API bool ImportStaticMeshHiResSourceModelDialog( UStaticMesh* StaticMesh );
 
+	/**
+	 * Remove the nanite data for the specified staticmesh
+	 *
+	 * @param Mesh				The mesh to import the high res mesh to
+	 * @return					Whether or not the import succeeded
+	 */
+	UNREALED_API bool RemoveStaticMeshHiRes(UStaticMesh* StaticMesh);
+
 	/**
 	 * Sets import option before importing
 	 *
diff --git a/Engine/Source/Editor/UnrealEd/Public/ImportUtils/StaticMeshImportUtils.h b/Engine/Source/Editor/UnrealEd/Public/ImportUtils/StaticMeshImportUtils.h
--- a/Engine/Source/Editor/UnrealEd/Public/ImportUtils/StaticMeshImportUtils.h
+++ b/Engine/Source/Editor/UnrealEd/Public/ImportUtils/StaticMeshImportUtils.h
@@ -40,6 +40,8 @@ struct FExistingStaticMeshData
 {
 	TArray<FStaticMaterial> 	ExistingMaterials;
 
+	FExistingLODMeshData		HiResSourceData;
+
 	FMeshSectionInfoMap			ExistingSectionInfoMap;
 	TArray<FExistingLODMeshData>	ExistingLODData;
 
