Nanite Hi-Res mesh is now built with the LOD0 build settings instead of the default ones
Gray out Fallback Triangle Percent and Fallback Relative Error when a model is specified in Source Import Filename

diff --git a/Engine/Source/Developer/MeshBuilder/Private/StaticMeshBuilder.cpp b/Engine/Source/Developer/MeshBuilder/Private/StaticMeshBuilder.cpp
--- a/Engine/Source/Developer/MeshBuilder/Private/StaticMeshBuilder.cpp
+++ b/Engine/Source/Developer/MeshBuilder/Private/StaticMeshBuilder.cpp
@@ -119,7 +119,7 @@ static bool BuildNanite(
 
 	FMeshDescription MeshDescription = *SourceModel.GetOrCacheMeshDescription();
 
-	FMeshBuildSettings& BuildSettings = SourceModel.BuildSettings;
+	FMeshBuildSettings& BuildSettings = StaticMesh->GetSourceModel(0).BuildSettings;
 	FStaticMeshLODResources& StaticMeshLOD = LODResources[0];
 
 	// compute tangents, lightmap UVs, etc
diff --git a/Engine/Source/Editor/StaticMeshEditor/Private/StaticMeshEditorTools.cpp b/Engine/Source/Editor/StaticMeshEditor/Private/StaticMeshEditorTools.cpp
--- a/Engine/Source/Editor/StaticMeshEditor/Private/StaticMeshEditorTools.cpp
+++ b/Engine/Source/Editor/StaticMeshEditor/Private/StaticMeshEditorTools.cpp
@@ -5225,7 +5225,7 @@ void FNaniteSettingsLayout::AddToDetailsPanel(IDetailLayoutBuilder& DetailBuilde
 	{
 		NaniteSettingsCategory.AddCustomRow( LOCTEXT("FallbackTrianglePercent", "Fallback Triangle Percent") )
 
-		.IsEnabled(TAttribute<bool>::Create(TAttribute<bool>::FGetter::CreateLambda([NaniteEnabledCheck]() -> bool {return NaniteEnabledCheck->IsChecked(); } )))
+		.IsEnabled(TAttribute<bool>::Create(TAttribute<bool>::FGetter::CreateLambda([this, NaniteEnabledCheck]() -> bool {return NaniteEnabledCheck->IsChecked() && IsHiResDataEmpty(); })))
 		.NameContent()
 		[
 			SNew(STextBlock)
@@ -5249,7 +5249,7 @@ void FNaniteSettingsLayout::AddToDetailsPanel(IDetailLayoutBuilder& DetailBuilde
 	{
 		NaniteSettingsCategory.AddCustomRow( LOCTEXT("FallbackRelativeError", "Fallback Relative Error") )
 
-		.IsEnabled(TAttribute<bool>::Create(TAttribute<bool>::FGetter::CreateLambda([NaniteEnabledCheck]() -> bool {return NaniteEnabledCheck->IsChecked(); } )))
+		.IsEnabled(TAttribute<bool>::Create(TAttribute<bool>::FGetter::CreateLambda([this, NaniteEnabledCheck]() -> bool {return NaniteEnabledCheck->IsChecked() && IsHiResDataEmpty(); } )))
 		.NameContent()
 		[
 			SNew(STextBlock)
diff --git a/Engine/Source/Runtime/Core/Private/UObject/FortniteMainBranchObjectVersion.cpp b/Engine/Source/Runtime/Core/Private/UObject/FortniteMainBranchObjectVersion.cpp
--- a/Engine/Source/Runtime/Core/Private/UObject/FortniteMainBranchObjectVersion.cpp
+++ b/Engine/Source/Runtime/Core/Private/UObject/FortniteMainBranchObjectVersion.cpp
@@ -13,7 +13,7 @@ TMap<FGuid, FGuid> FFortniteMainBranchObjectVersion::GetSystemGuids()
 	SystemGuids.Add(DevGuids.NIAGARASHADERMAP_DERIVEDDATA_VER, FGuid("2289D5116CF94BC0AFEAEC541468E645"));
 	SystemGuids.Add(DevGuids.Niagara_LatestScriptCompileVersion, FGuid("D6A13E2AE4D2453EB1A44529AAAB8305"));
 	SystemGuids.Add(DevGuids.SkeletalMeshDerivedDataVersion, FGuid("ACF593EAAE354FCBB34CF44F5AA1BFD2"));
-	SystemGuids.Add(DevGuids.STATICMESH_DERIVEDDATA_VER, FGuid("4696D39D9BC743AFBFA1DAB359AD7A7D"));
+	SystemGuids.Add(DevGuids.STATICMESH_DERIVEDDATA_VER, FGuid("9A48CE6B46204777968322F98F82D11F"));
 
 	return SystemGuids;
 }
