Fixed static mesh editor so that the currently visible LOD is displayed when rendering Nanite in fallback mode.

diff --git a/Engine/Source/Editor/StaticMeshEditor/Private/StaticMeshEditorViewportClient.cpp b/Engine/Source/Editor/StaticMeshEditor/Private/StaticMeshEditorViewportClient.cpp
--- a/Engine/Source/Editor/StaticMeshEditor/Private/StaticMeshEditorViewportClient.cpp
+++ b/Engine/Source/Editor/StaticMeshEditor/Private/StaticMeshEditorViewportClient.cpp
@@ -732,7 +732,7 @@ static void DrawAngles(FCanvas* Canvas, int32 XPos, int32 YPos, EAxisList::Type
 	}
 }
 
-void FStaticMeshEditorViewportClient::DrawCanvas( FViewport& InViewport, FSceneView& View, FCanvas& Canvas )
+void FStaticMeshEditorViewportClient::DrawCanvas(FViewport& InViewport, FSceneView& View, FCanvas& Canvas)
 {
 	auto StaticMeshEditor = StaticMeshEditorPtr.Pin();
 	auto StaticMeshEditorViewport = StaticMeshEditorViewportPtr.Pin();
@@ -741,28 +741,28 @@ void FStaticMeshEditorViewportClient::DrawCanvas( FViewport& InViewport, FSceneV
 		return;
 	}
 
-	const int32 HalfX = Viewport->GetSizeXY().X/2 / GetDPIScale();
-	const int32 HalfY = Viewport->GetSizeXY().Y/2 / GetDPIScale();
+	const int32 HalfX = Viewport->GetSizeXY().X / 2 / GetDPIScale();
+	const int32 HalfY = Viewport->GetSizeXY().Y / 2 / GetDPIScale();
 
 	// Draw socket names if desired.
-	if( bShowSockets )
+	if (bShowSockets)
 	{
-		for(int32 i=0; i<StaticMesh->Sockets.Num(); i++)
+		for (int32 i = 0; i < StaticMesh->Sockets.Num(); i++)
 		{
 			UStaticMeshSocket* Socket = StaticMesh->Sockets[i];
-			if(Socket!=NULL)
+			if (Socket != nullptr)
 			{
 				FMatrix SocketTM;
 				Socket->GetSocketMatrix(SocketTM, StaticMeshComponent);
-				const FVector SocketPos	= SocketTM.GetOrigin();
-				const FPlane proj		= View.Project( SocketPos );
-				if(proj.W > 0.f)
+				const FVector SocketPos = SocketTM.GetOrigin();
+				const FPlane proj = View.Project(SocketPos);
+				if (proj.W > 0.f)
 				{
-					const int32 XPos = HalfX + ( HalfX * proj.X );
-					const int32 YPos = HalfY + ( HalfY * (proj.Y * -1) );
+					const int32 XPos = HalfX + (HalfX * proj.X);
+					const int32 YPos = HalfY + (HalfY * (proj.Y * -1));
 
-					FCanvasTextItem TextItem( FVector2D( XPos, YPos ), FText::FromString( Socket->SocketName.ToString() ), GEngine->GetSmallFont(), FLinearColor(FColor(255,196,196)) );
-					Canvas.DrawItem( TextItem );	
+					FCanvasTextItem TextItem(FVector2D(XPos, YPos), FText::FromString(Socket->SocketName.ToString()), GEngine->GetSmallFont(), FLinearColor(FColor(255, 196, 196)));
+					Canvas.DrawItem(TextItem);
 
 					const UStaticMeshSocket* SelectedSocket = StaticMeshEditor->GetSelectedSocket();
 					if (bManipulating && SelectedSocket == Socket)
@@ -772,8 +772,8 @@ void FStaticMeshEditorViewportClient::DrawCanvas( FViewport& InViewport, FSceneV
 						UCanvas::CanvasStringSize(Parameters, *Socket->SocketName.ToString());
 						int32 YL = FMath::TruncToInt(Parameters.DrawYL);
 
-						DrawAngles(&Canvas, XPos, YPos + YL, 
-							Widget->GetCurrentAxis(), 
+						DrawAngles(&Canvas, XPos, YPos + YL,
+							Widget->GetCurrentAxis(),
 							GetWidgetMode(),
 							Socket->RelativeRotation,
 							Socket->RelativeLocation);
@@ -805,7 +805,7 @@ void FStaticMeshEditorViewportClient::DrawCanvas( FViewport& InViewport, FSceneV
 			{
 				const FString PositionStr = FNaniteSettingsLayout::PositionPrecisionValueToDisplayString(Resources.PositionPrecision);
 				TextItems.Emplace(FText::Format(NSLOCTEXT("UnrealEd", "NanitePositionPrecision", "Position Precision: {0}"), FText::FromString(PositionStr)));
-				
+
 				const uint32 NumStreamingPages = Resources.PageStreamingStates.Num() - Resources.NumRootPages;
 				const uint64 RootKB = uint64(Resources.NumRootPages) * NANITE_ROOT_PAGE_GPU_SIZE;
 				const uint64 StreamingKB = uint64(NumStreamingPages) * NANITE_STREAMING_PAGE_GPU_SIZE;
@@ -820,13 +820,13 @@ void FStaticMeshEditorViewportClient::DrawCanvas( FViewport& InViewport, FSceneV
 					FText::AsNumber(RootKB / 1048576.0f, &NumberOptions),
 					FText::AsNumber(StreamingKB / 1048576.0f, &NumberOptions),
 					FText::AsNumber(TotalKB / 1048576.0f, &NumberOptions)
-					));
+				));
 			}
 		}
 	}
-	else
+
+	if (!StaticMesh->NaniteSettings.bEnabled || StaticMeshComponent->bDisplayNaniteFallbackMesh)
 	{
-		
 		const int32 CurrentMinLODLevel = StaticMesh->GetMinLOD().GetValue();
 		const bool bBelowMinLOD = CurrentLODLevel < CurrentMinLODLevel;
 
@@ -876,7 +876,7 @@ void FStaticMeshEditorViewportClient::DrawCanvas( FViewport& InViewport, FSceneV
 
 	TextItems.Emplace(FText::Format(NSLOCTEXT("UnrealEd", "UVChannels_F", "UV Channels:  {0}"), FText::AsNumber(StaticMeshEditorPtr.Pin()->GetNumUVChannels(CurrentLODLevel))));
 
-	if(StaticMesh->GetRenderData() && StaticMesh->GetRenderData()->LODResources.Num() > 0 )
+	if (StaticMesh->GetRenderData() && StaticMesh->GetRenderData()->LODResources.Num() > 0)
 	{
 		if (StaticMesh->GetRenderData()->LODResources[0].DistanceFieldData != nullptr )
 		{
@@ -902,13 +902,13 @@ void FStaticMeshEditorViewportClient::DrawCanvas( FViewport& InViewport, FSceneV
 		FText::AsNumber(int32(StaticMesh->GetBounds().BoxExtent.Z * 2.0f))));
 
 	// Show the number of collision primitives
-	if(StaticMesh->GetBodySetup())
+	if (StaticMesh->GetBodySetup())
 	{
 		TextItems.Emplace(FText::Format(NSLOCTEXT("UnrealEd", "NumPrimitives_F", "Num Collision Primitives:  {0}"), FText::AsNumber(StaticMesh->GetBodySetup()->AggGeom.GetElementCount())));
 	}
 
 	// Estimated compressed size
-	if(StaticMesh->GetRenderData())
+	if (StaticMesh->GetRenderData())
 	{
 		FNumberFormattingOptions NumberOptions;
 		NumberOptions.MinimumFractionalDigits = 2;
@@ -936,7 +936,7 @@ void FStaticMeshEditorViewportClient::DrawCanvas( FViewport& InViewport, FSceneV
 		Canvas.DrawShadowedText(X, Y, Message, GetStatsFont(), FLinearColor(1.0, 0.05, 0.05, 1.0));
 	}
 
-	if(bDrawUVs && StaticMesh->GetRenderData()->LODResources.Num() > 0)
+	if (bDrawUVs && StaticMesh->GetRenderData()->LODResources.Num() > 0)
 	{
 		const int32 YPos = 160;
 		DrawUVsForMesh(Viewport, &Canvas, YPos);
