StaticMeshEditor: fix implicit cast warnings

diff --git a/Engine/Source/Editor/StaticMeshEditor/Private/SStaticMeshEditorViewport.cpp b/Engine/Source/Editor/StaticMeshEditor/Private/SStaticMeshEditorViewport.cpp
--- a/Engine/Source/Editor/StaticMeshEditor/Private/SStaticMeshEditorViewport.cpp
+++ b/Engine/Source/Editor/StaticMeshEditor/Private/SStaticMeshEditorViewport.cpp
@@ -40,7 +40,7 @@ void SStaticMeshEditorViewport::Construct(const FArguments& InArgs)
 
 	if (StaticMesh)
 	{
-		PreviewScene->SetFloorOffset(-StaticMesh->GetExtendedBounds().Origin.Z + StaticMesh->GetExtendedBounds().BoxExtent.Z);
+		PreviewScene->SetFloorOffset(static_cast<float>( -StaticMesh->GetExtendedBounds().Origin.Z + StaticMesh->GetExtendedBounds().BoxExtent.Z ));
 	}
 
 	// restore last used feature level
diff --git a/Engine/Source/Editor/StaticMeshEditor/Private/StaticMeshEditor.cpp b/Engine/Source/Editor/StaticMeshEditor/Private/StaticMeshEditor.cpp
--- a/Engine/Source/Editor/StaticMeshEditor/Private/StaticMeshEditor.cpp
+++ b/Engine/Source/Editor/StaticMeshEditor/Private/StaticMeshEditor.cpp
@@ -1297,7 +1297,7 @@ void FStaticMeshEditor::ScaleSelectedPrims(const FVector& InScale)
 	}
 
 	//Multiply in estimated size of the mesh so scaling of sphere, box and sphyl is similar speed to other scaling
-	float SimplePrimitiveScaleSpeedFactor = StaticMesh->GetBounds().SphereRadius;
+	float SimplePrimitiveScaleSpeedFactor = static_cast<float>( StaticMesh->GetBounds().SphereRadius );
 
 	for (int32 PrimIdx = 0; PrimIdx < SelectedPrims.Num(); PrimIdx++)
 	{
diff --git a/Engine/Source/Editor/StaticMeshEditor/Private/StaticMeshEditorSubsystem.cpp b/Engine/Source/Editor/StaticMeshEditor/Private/StaticMeshEditorSubsystem.cpp
--- a/Engine/Source/Editor/StaticMeshEditor/Private/StaticMeshEditorSubsystem.cpp
+++ b/Engine/Source/Editor/StaticMeshEditor/Private/StaticMeshEditorSubsystem.cpp
@@ -1392,13 +1392,13 @@ bool UStaticMeshEditorSubsystem::BulkSetConvexDecompositionCollisionsWithNotific
 
 	uint32 LastProcessed = 0;
 	const FText ProgressText = LOCTEXT("ComputingConvexCollision", "Computing convex collision for static mesh {0}/{1} ...");
-	FScopedSlowTask Progress(StaticMeshes.Num(), FText::Format(ProgressText, LastProcessed, StaticMeshes.Num()));
+	FScopedSlowTask Progress(static_cast<float>(StaticMeshes.Num()), FText::Format(ProgressText, LastProcessed, StaticMeshes.Num()));
 	Progress.MakeDialog();
 
 	while (!Result.WaitFor(FTimespan::FromMilliseconds(33.0)))
 	{
 		uint32 LocalProcessed = Processed.Load(EMemoryOrder::Relaxed);
-		Progress.EnterProgressFrame(LocalProcessed - LastProcessed, FText::Format(ProgressText, LocalProcessed, StaticMeshes.Num()));
+		Progress.EnterProgressFrame(static_cast<float>(LocalProcessed - LastProcessed), FText::Format(ProgressText, LocalProcessed, StaticMeshes.Num()));
 		LastProcessed = LocalProcessed;
 	}
 
diff --git a/Engine/Source/Editor/StaticMeshEditor/Private/StaticMeshEditorTools.cpp b/Engine/Source/Editor/StaticMeshEditor/Private/StaticMeshEditorTools.cpp
--- a/Engine/Source/Editor/StaticMeshEditor/Private/StaticMeshEditorTools.cpp
+++ b/Engine/Source/Editor/StaticMeshEditor/Private/StaticMeshEditorTools.cpp
@@ -922,17 +922,17 @@ int32 FMeshBuildSettingsLayout::GetDstLightmapIndex() const
 
 TOptional<float> FMeshBuildSettingsLayout::GetBuildScaleX() const
 {
-	return BuildSettings.BuildScale3D.X;
+	return static_cast<float>(BuildSettings.BuildScale3D.X);
 }
 
 TOptional<float> FMeshBuildSettingsLayout::GetBuildScaleY() const
 {
-	return BuildSettings.BuildScale3D.Y;
+	return static_cast<float>(BuildSettings.BuildScale3D.Y);
 }
 
 TOptional<float> FMeshBuildSettingsLayout::GetBuildScaleZ() const
 {
-	return BuildSettings.BuildScale3D.Z;
+	return static_cast<float>(BuildSettings.BuildScale3D.Z);
 }
 
 float FMeshBuildSettingsLayout::GetDistanceFieldResolutionScale() const
@@ -1386,7 +1386,7 @@ void FMeshReductionSettingsLayout::GenerateChildContent( IDetailChildrenBuilder&
 				[
 					SAssignNew(SilhouetteCombo, STextComboBox)
 					//.Font( IDetailLayoutBuilder::GetDetailFont() )
-				.ContentPadding(0)
+				.ContentPadding(0.f)
 				.OptionsSource(&ImportanceOptions)
 				.InitiallySelectedItem(ImportanceOptions[ReductionSettings.SilhouetteImportance])
 				.OnSelectionChanged(this, &FMeshReductionSettingsLayout::OnSilhouetteImportanceChanged)
@@ -1406,7 +1406,7 @@ void FMeshReductionSettingsLayout::GenerateChildContent( IDetailChildrenBuilder&
 				[
 					SAssignNew(TextureCombo, STextComboBox)
 					//.Font( IDetailLayoutBuilder::GetDetailFont() )
-				.ContentPadding(0)
+				.ContentPadding(0.f)
 				.OptionsSource(&ImportanceOptions)
 				.InitiallySelectedItem(ImportanceOptions[ReductionSettings.TextureImportance])
 				.OnSelectionChanged(this, &FMeshReductionSettingsLayout::OnTextureImportanceChanged)
@@ -1426,7 +1426,7 @@ void FMeshReductionSettingsLayout::GenerateChildContent( IDetailChildrenBuilder&
 				[
 					SAssignNew(ShadingCombo, STextComboBox)
 					//.Font( IDetailLayoutBuilder::GetDetailFont() )
-				.ContentPadding(0)
+				.ContentPadding(0.f)
 				.OptionsSource(&ImportanceOptions)
 				.InitiallySelectedItem(ImportanceOptions[ReductionSettings.ShadingImportance])
 				.OnSelectionChanged(this, &FMeshReductionSettingsLayout::OnShadingImportanceChanged)
@@ -3963,7 +3963,7 @@ TSharedRef<SWidget> FLevelOfDetailSettingsLayout::GetLODScreenSizeWidget(FName P
 		.Font(IDetailLayoutBuilder::GetDetailFont())
 		.MinDesiredWidth(60.0f)
 		.MinValue(0.0f)
-		.MaxValue(WORLD_MAX)
+		.MaxValue(static_cast<float>(WORLD_MAX))
 		.SliderExponent(2.0f)
 		.Value(this, &FLevelOfDetailSettingsLayout::GetLODScreenSize, PlatformGroupName, LODIndex)
 		.OnValueChanged(const_cast<FLevelOfDetailSettingsLayout*>(this), &FLevelOfDetailSettingsLayout::OnLODScreenSizeChanged, PlatformGroupName, LODIndex)
@@ -4432,7 +4432,7 @@ bool FLevelOfDetailSettingsLayout::AddMinLODPlatformOverride(FName PlatformGroup
 	if (StaticMesh->GetMinLOD().PerPlatform.Find(PlatformGroupName) == nullptr)
 	{
 		FPerPlatformInt MinLOD = StaticMesh->GetMinLOD();
-		float Value = MinLOD.Default;
+		int32 Value = MinLOD.Default;
 		MinLOD.PerPlatform.Add(PlatformGroupName, Value);
 		StaticMesh->SetMinLOD(MoveTemp(MinLOD));
 		OnMinLODChanged(Value, PlatformGroupName);
@@ -4450,7 +4450,7 @@ bool FLevelOfDetailSettingsLayout::RemoveMinLODPlatformOverride(FName PlatformGr
 	FPerPlatformInt MinLOD = StaticMesh->GetMinLOD();
 	if (MinLOD.PerPlatform.Remove(PlatformGroupName) != 0)
 	{
-		float Value = MinLOD.Default;
+		int32 Value = MinLOD.Default;
 		StaticMesh->SetMinLOD(MoveTemp(MinLOD));
 		OnMinLODChanged(Value, PlatformGroupName);
 		return true;
@@ -4535,7 +4535,7 @@ bool FLevelOfDetailSettingsLayout::AddMinLODQualityLevelOverride(FName QualityLe
 	if (StaticMesh->GetQualityLevelMinLOD().PerQuality.Find(QLKey) == nullptr)
 	{
 		FPerQualityLevelInt MinLOD = StaticMesh->GetQualityLevelMinLOD();
-		float Value = MinLOD.Default;
+		int32 Value = MinLOD.Default;
 		MinLOD.PerQuality.Add(QLKey, Value);
 		StaticMesh->SetQualityLevelMinLOD(MoveTemp(MinLOD));
 		OnMinQualityLevelLODChanged(Value, QualityLevelName);
@@ -4555,7 +4555,7 @@ bool FLevelOfDetailSettingsLayout::RemoveMinLODQualityLevelOverride(FName Qualit
 	int32 QL = QualityLevelProperty::FNameToQualityLevel(QualityLevelName);
 	if (QL != INDEX_NONE && MinLOD.PerQuality.Remove(QL) != 0)
 	{
-		float Value = MinLOD.Default;
+		int32 Value = MinLOD.Default;
 		StaticMesh->SetQualityLevelMinLOD(MoveTemp(MinLOD));
 		OnMinQualityLevelLODChanged(Value, QualityLevelName);
 		return true;
@@ -4787,7 +4787,7 @@ bool FLevelOfDetailSettingsLayout::AddNumStreamedLODsPlatformOverride(FName Plat
 	StaticMesh->Modify();
 	if (StaticMesh->NumStreamedLODs.PerPlatform.Find(PlatformGroupName) == nullptr)
 	{
-		float Value = StaticMesh->NumStreamedLODs.Default;
+		int32 Value = StaticMesh->NumStreamedLODs.Default;
 		StaticMesh->NumStreamedLODs.PerPlatform.Add(PlatformGroupName, Value);
 		OnNumStreamedLODsChanged(Value, PlatformGroupName);
 		return true;
@@ -5474,7 +5474,7 @@ FString FNaniteSettingsLayout::PositionPrecisionValueToDisplayString(int32 Value
 	}
 	else
 	{
-		const float fValue = FMath::Exp2((double)-Value);
+		const float fValue = static_cast<float>(FMath::Exp2((double)-Value));
 		return FString::Printf(TEXT("1/%dcm (%.3gcm)"), 1 << Value, fValue);
 	}
 }
diff --git a/Engine/Source/Editor/StaticMeshEditor/Private/StaticMeshEditorViewportClient.cpp b/Engine/Source/Editor/StaticMeshEditor/Private/StaticMeshEditorViewportClient.cpp
--- a/Engine/Source/Editor/StaticMeshEditor/Private/StaticMeshEditorViewportClient.cpp
+++ b/Engine/Source/Editor/StaticMeshEditor/Private/StaticMeshEditorViewportClient.cpp
@@ -62,7 +62,7 @@ FStaticMeshEditorViewportClient::FStaticMeshEditorViewportClient(TWeakPtr<IStati
 	DrawHelper.GridColorMajor = FColor(144,144,144);
 	DrawHelper.GridColorMinor = FColor(128,128,128);
 	DrawHelper.PerspectiveGridSize = GridSize;
-	DrawHelper.NumCells = DrawHelper.PerspectiveGridSize / (CellSize * 2);
+	DrawHelper.NumCells = FMath::FloorToInt32(DrawHelper.PerspectiveGridSize / (CellSize * 2));
 
 	SetViewMode(VMI_Lit);
 
@@ -741,8 +741,8 @@ void FStaticMeshEditorViewportClient::DrawCanvas(FViewport& InViewport, FSceneVi
 		return;
 	}
 
-	const int32 HalfX = Viewport->GetSizeXY().X / 2 / GetDPIScale();
-	const int32 HalfY = Viewport->GetSizeXY().Y / 2 / GetDPIScale();
+	const int32 HalfX = FMath::FloorToInt32(Viewport->GetSizeXY().X / 2 / GetDPIScale());
+	const int32 HalfY = FMath::FloorToInt32(Viewport->GetSizeXY().Y / 2 / GetDPIScale());
 
 	// Draw socket names if desired.
 	if (bShowSockets)
@@ -758,8 +758,8 @@ void FStaticMeshEditorViewportClient::DrawCanvas(FViewport& InViewport, FSceneVi
 				const FPlane proj = View.Project(SocketPos);
 				if (proj.W > 0.f)
 				{
-					const int32 XPos = HalfX + (HalfX * proj.X);
-					const int32 YPos = HalfY + (HalfY * (proj.Y * -1));
+					const int32 XPos = FMath::FloorToInt32( HalfX + (HalfX * proj.X) );
+					const int32 YPos = FMath::FloorToInt32( HalfY + (HalfY * (proj.Y * -1)) );
 
 					FCanvasTextItem TextItem(FVector2D(XPos, YPos), FText::FromString(Socket->SocketName.ToString()), GEngine->GetSmallFont(), FLinearColor(FColor(255, 196, 196)));
 					Canvas.DrawItem(TextItem);
@@ -789,7 +789,7 @@ void FStaticMeshEditorViewportClient::DrawCanvas(FViewport& InViewport, FSceneVi
 	{
 		int32 LOD = StaticMeshEditor->GetCurrentLODLevel();
 		return (LOD == 0) ?
-			ComputeStaticMeshLOD(StaticMesh->GetRenderData(), StaticMeshComponent->Bounds.Origin, StaticMeshComponent->Bounds.SphereRadius, View, StaticMesh->GetDefaultMinLOD())
+			ComputeStaticMeshLOD(StaticMesh->GetRenderData(), StaticMeshComponent->Bounds.Origin, static_cast<float>(StaticMeshComponent->Bounds.SphereRadius), View, StaticMesh->GetDefaultMinLOD())
 			:
 			LOD - 1;
 	}();
@@ -838,7 +838,7 @@ void FStaticMeshEditorViewportClient::DrawCanvas(FViewport& InViewport, FSceneVi
 		}
 	}
 
-	const float CurrentScreenSize = ComputeBoundsScreenSize(StaticMeshComponent->Bounds.Origin, StaticMeshComponent->Bounds.SphereRadius, View);
+	const float CurrentScreenSize = ComputeBoundsScreenSize(StaticMeshComponent->Bounds.Origin, static_cast<float>(StaticMeshComponent->Bounds.SphereRadius), View);
 	FNumberFormattingOptions FormatOptions;
 	FormatOptions.MinimumFractionalDigits = 3;
 	FormatOptions.MaximumFractionalDigits = 6;
@@ -933,7 +933,7 @@ void FStaticMeshEditorViewportClient::DrawCanvas(FViewport& InViewport, FSceneVi
 	if (StaticMesh->GetBodySetup() && (!(StaticMesh->GetBodySetup()->bHasCookedCollisionData || StaticMesh->GetBodySetup()->bNeverNeedsCookedCollisionData) || StaticMesh->GetBodySetup()->bFailedToCreatePhysicsMeshes))
 	{
 		static const FText Message = NSLOCTEXT("Renderer", "NoCookedCollisionObject", "NO COOKED COLLISION OBJECT: TOO SMALL?");
-		Canvas.DrawShadowedText(X, Y, Message, GetStatsFont(), FLinearColor(1.0, 0.05, 0.05, 1.0));
+		Canvas.DrawShadowedText(static_cast<float>(X), static_cast<float>(Y), Message, GetStatsFont(), FLinearColor(1.0, 0.05, 0.05, 1.0));
 	}
 
 	if (bDrawUVs && StaticMesh->GetRenderData()->LODResources.Num() > 0)
@@ -1199,7 +1199,7 @@ void FStaticMeshEditorViewportClient::ProcessClick(class FSceneView& InView, cla
 										ClosestPointToClickLineOnEdge);
 
 									// Compute the minimum distance (squared)
-									const float MinDistanceToEdgeSquared = (ClosestPointToClickLineOnEdge - ClosestPointToEdgeOnClickLine).SizeSquared();
+									const float MinDistanceToEdgeSquared = static_cast<float>( (ClosestPointToClickLineOnEdge - ClosestPointToEdgeOnClickLine).SizeSquared() );
 
 									if (MinDistanceToEdgeSquared <= WorldSpaceMinClickDistance)
 									{
@@ -1442,7 +1442,7 @@ void FStaticMeshEditorViewportClient::SetPreviewMesh(UStaticMesh* InStaticMesh,
 		ThumbnailAngle.Roll = 0;
 		const float ThumbnailDistance = ThumbnailInfo->OrbitZoom;
 
-		const float CameraY = StaticMesh->GetBounds().SphereRadius / (75.0f * PI / 360.0f);
+		const float CameraY = static_cast<float>( StaticMesh->GetBounds().SphereRadius / (75.0f * PI / 360.0f) );
 		SetCameraSetup(
 			FVector::ZeroVector,
 			ThumbnailAngle,
