Set the ParentTab for StaticMeshEditorViewport on spawning that tab, and also check SEditorViewport::IsVisible for StaticMeshEditorViewport visibility, as it takes into consideration last time on screen.

diff --git a/Engine/Source/Editor/StaticMeshEditor/Private/SStaticMeshEditorViewport.cpp b/Engine/Source/Editor/StaticMeshEditor/Private/SStaticMeshEditorViewport.cpp
--- a/Engine/Source/Editor/StaticMeshEditor/Private/SStaticMeshEditorViewport.cpp
+++ b/Engine/Source/Editor/StaticMeshEditor/Private/SStaticMeshEditorViewport.cpp
@@ -361,7 +361,7 @@ void SStaticMeshEditorViewport::UpdatePreviewMesh(UStaticMesh* InStaticMesh, boo
 
 bool SStaticMeshEditorViewport::IsVisible() const
 {
-	return ViewportWidget.IsValid() && (!ParentTab.IsValid() || ParentTab.Pin()->IsForeground());
+	return ViewportWidget.IsValid() && (!ParentTab.IsValid() || ParentTab.Pin()->IsForeground()) && SEditorViewport::IsVisible();
 }
 
 UStaticMeshComponent* SStaticMeshEditorViewport::GetStaticMeshComponent() const
diff --git a/Engine/Source/Editor/StaticMeshEditor/Private/StaticMeshEditor.cpp b/Engine/Source/Editor/StaticMeshEditor/Private/StaticMeshEditor.cpp
--- a/Engine/Source/Editor/StaticMeshEditor/Private/StaticMeshEditor.cpp
+++ b/Engine/Source/Editor/StaticMeshEditor/Private/StaticMeshEditor.cpp
@@ -547,6 +547,8 @@ TSharedRef<SDockTab> FStaticMeshEditor::SpawnTab_Viewport( const FSpawnTabArgs&
 	const FString LayoutId = FString("StaticMeshEditorViewport");
 	ViewportTabContent->Initialize(MakeViewportFunc, DockableTab, LayoutId);
 
+	GetStaticMeshViewport()->SetParentTab(DockableTab);
+
 	return DockableTab;
 }
 
