Add RowTags for StaticMeshBuildSettings customization

diff --git a/Engine/Source/Editor/StaticMeshEditor/Private/StaticMeshEditorTools.cpp b/Engine/Source/Editor/StaticMeshEditor/Private/StaticMeshEditorTools.cpp
--- a/Engine/Source/Editor/StaticMeshEditor/Private/StaticMeshEditorTools.cpp
+++ b/Engine/Source/Editor/StaticMeshEditor/Private/StaticMeshEditorTools.cpp
@@ -489,6 +489,7 @@ void FMeshBuildSettingsLayout::GenerateChildContent( IDetailChildrenBuilder& Chi
 {
 	{
 		ChildrenBuilder.AddCustomRow( LOCTEXT("RecomputeNormals", "Recompute Normals") )
+		.RowTag("RecomputeNormals")
 		.NameContent()
 		[
 			SNew(STextBlock)
@@ -506,6 +507,7 @@ void FMeshBuildSettingsLayout::GenerateChildContent( IDetailChildrenBuilder& Chi
 
 	{
 		ChildrenBuilder.AddCustomRow( LOCTEXT("RecomputeTangents", "Recompute Tangents") )
+		.RowTag("RecomputeTangents")
 		.NameContent()
 		[
 			SNew(STextBlock)
@@ -522,6 +524,7 @@ void FMeshBuildSettingsLayout::GenerateChildContent( IDetailChildrenBuilder& Chi
 
 	{
 		ChildrenBuilder.AddCustomRow( LOCTEXT("UseMikkTSpace", "Use MikkTSpace Tangent Space") )
+		.RowTag("UseMikkTSpace")
 		.NameContent()
 		[
 			SNew(STextBlock)
@@ -538,6 +541,7 @@ void FMeshBuildSettingsLayout::GenerateChildContent( IDetailChildrenBuilder& Chi
 
 	{
 		ChildrenBuilder.AddCustomRow( LOCTEXT("ComputeWeightedNormals", "Compute Weighted Normals") )
+		.RowTag("ComputeWeightedNormals")
 		.NameContent()
 		[
 			SNew(STextBlock)
@@ -554,6 +558,7 @@ void FMeshBuildSettingsLayout::GenerateChildContent( IDetailChildrenBuilder& Chi
 
 	{
 		ChildrenBuilder.AddCustomRow( LOCTEXT("RemoveDegenerates", "Remove Degenerates") )
+		.RowTag("RemoveDegenerates")
 		.NameContent()
 		[
 			SNew(STextBlock)
@@ -571,6 +576,7 @@ void FMeshBuildSettingsLayout::GenerateChildContent( IDetailChildrenBuilder& Chi
 
 	{
 		ChildrenBuilder.AddCustomRow( LOCTEXT("BuildReversedIndexBuffer", "Build Reversed Index Buffer") )
+		.RowTag("BuildReversedIndexBuffer")
 		.NameContent()
 		[
 			SNew(STextBlock)
@@ -587,6 +593,7 @@ void FMeshBuildSettingsLayout::GenerateChildContent( IDetailChildrenBuilder& Chi
 
 	{
 		ChildrenBuilder.AddCustomRow(LOCTEXT("UseHighPrecisionTangentBasis", "Use High Precision Tangent Basis"))
+		.RowTag("UseHighPrecisionTangentBasis")
 		.NameContent()
 		[
 			SNew(STextBlock)
@@ -603,6 +610,7 @@ void FMeshBuildSettingsLayout::GenerateChildContent( IDetailChildrenBuilder& Chi
 
 	{
 		ChildrenBuilder.AddCustomRow( LOCTEXT("UseFullPrecisionUVs", "Use Full Precision UVs") )
+		.RowTag("UseFullPrecisionUVs")
 		.NameContent()
 		[
 			SNew(STextBlock)
@@ -619,6 +627,7 @@ void FMeshBuildSettingsLayout::GenerateChildContent( IDetailChildrenBuilder& Chi
 	
 	{
 		ChildrenBuilder.AddCustomRow( LOCTEXT("UseBackwardsCompatibleF16TruncUVs", "UE4 Compatible UVs") )
+		.RowTag("UseBackwardsCompatibleF16TruncUVs")
 		.NameContent()
 		[
 			SNew(STextBlock)
@@ -635,6 +644,7 @@ void FMeshBuildSettingsLayout::GenerateChildContent( IDetailChildrenBuilder& Chi
 
 	{
 		ChildrenBuilder.AddCustomRow( LOCTEXT("GenerateLightmapUVs", "Generate Lightmap UVs") )
+		.RowTag("GenerateLightmapUVs")
 		.NameContent()
 		[
 			SNew(STextBlock)
@@ -651,6 +661,7 @@ void FMeshBuildSettingsLayout::GenerateChildContent( IDetailChildrenBuilder& Chi
 
 	{
 		ChildrenBuilder.AddCustomRow( LOCTEXT("MinLightmapResolution", "Min Lightmap Resolution") )
+		.RowTag("MinLightmapResolution")
 		.NameContent()
 		[
 			SNew(STextBlock)
@@ -670,6 +681,7 @@ void FMeshBuildSettingsLayout::GenerateChildContent( IDetailChildrenBuilder& Chi
 
 	{
 		ChildrenBuilder.AddCustomRow( LOCTEXT("SourceLightmapIndex", "Source Lightmap Index") )
+		.RowTag("SourceLightmapIndex")
 		.NameContent()
 		[
 			SNew(STextBlock)
@@ -689,6 +701,7 @@ void FMeshBuildSettingsLayout::GenerateChildContent( IDetailChildrenBuilder& Chi
 
 	{
 		ChildrenBuilder.AddCustomRow( LOCTEXT("DestinationLightmapIndex", "Destination Lightmap Index") )
+		.RowTag("DestinationLightmapIndex")
 		.NameContent()
 		[
 			SNew(STextBlock)
@@ -708,6 +721,7 @@ void FMeshBuildSettingsLayout::GenerateChildContent( IDetailChildrenBuilder& Chi
 
 	{
 		ChildrenBuilder.AddCustomRow(LOCTEXT("BuildScale", "Build Scale"))
+		.RowTag("BuildScale")
 		.NameContent()
 		[
 			SNew(STextBlock)
@@ -734,6 +748,7 @@ void FMeshBuildSettingsLayout::GenerateChildContent( IDetailChildrenBuilder& Chi
 
 	{
 		ChildrenBuilder.AddCustomRow( LOCTEXT("DistanceFieldResolutionScale", "Distance Field Resolution Scale") )
+		.RowTag("DistanceFieldResolutionScale")
 		.NameContent()
 		[
 			SNew(STextBlock)
@@ -754,6 +769,7 @@ void FMeshBuildSettingsLayout::GenerateChildContent( IDetailChildrenBuilder& Chi
 		
 	{
 		ChildrenBuilder.AddCustomRow( LOCTEXT("GenerateDistanceFieldAsIfTwoSided", "Two-Sided Distance Field Generation") )
+		.RowTag("GenerateDistanceFieldAsIfTwoSided")
 		.NameContent()
 		[
 			SNew(STextBlock)
@@ -776,6 +792,7 @@ void FMeshBuildSettingsLayout::GenerateChildContent( IDetailChildrenBuilder& Chi
 			.OnObjectChanged(this, &FMeshBuildSettingsLayout::OnDistanceFieldReplacementMeshSelected);
 
 		ChildrenBuilder.AddCustomRow( LOCTEXT("DistanceFieldReplacementMesh", "Distance Field Replacement Mesh") )
+		.RowTag("DistanceFieldReplacementMesh")
 		.NameContent()
 		[
 			SNew(STextBlock)
@@ -790,6 +807,7 @@ void FMeshBuildSettingsLayout::GenerateChildContent( IDetailChildrenBuilder& Chi
 
 	{
 		ChildrenBuilder.AddCustomRow(LOCTEXT("MaxLumenMeshCards", "Max Lumen Mesh Cards"))
+		.RowTag("MaxLumenMeshCards")
 			.NameContent()
 			[
 				SNew(STextBlock)
