Add Absolute criterion to staticmesh FMeshReductionSettings.

#rnx

diff --git a/Engine/Plugins/Experimental/GeometryScripting/Source/GeometryScriptingCore/Private/MeshAssetFunctions.cpp b/Engine/Plugins/Experimental/GeometryScripting/Source/GeometryScriptingCore/Private/MeshAssetFunctions.cpp
--- a/Engine/Plugins/Experimental/GeometryScripting/Source/GeometryScriptingCore/Private/MeshAssetFunctions.cpp
+++ b/Engine/Plugins/Experimental/GeometryScripting/Source/GeometryScriptingCore/Private/MeshAssetFunctions.cpp
@@ -330,8 +330,7 @@ UDynamicMesh*  UGeometryScriptLibrary_StaticMeshFunctions::CopyMeshToStaticMesh(
 
 		// Setting to prevent the standard static mesh reduction from running and replacing the render LOD.
 		FStaticMeshSourceModel& ThisSourceModel = ToStaticMeshAsset->GetSourceModel(UseLODIndex);
-		ThisSourceModel.ReductionSettings.PercentTriangles = 1.f;
-		ThisSourceModel.ReductionSettings.PercentVertices = 1.f;
+		ThisSourceModel.ResetReductionSetting();
 
 		if (Options.bApplyNaniteSettings)
 		{
diff --git a/Engine/Plugins/Experimental/SkeletalReduction/Source/Private/SkeletalMeshReductionPlugin.cpp b/Engine/Plugins/Experimental/SkeletalReduction/Source/Private/SkeletalMeshReductionPlugin.cpp
--- a/Engine/Plugins/Experimental/SkeletalReduction/Source/Private/SkeletalMeshReductionPlugin.cpp
+++ b/Engine/Plugins/Experimental/SkeletalReduction/Source/Private/SkeletalMeshReductionPlugin.cpp
@@ -71,6 +71,11 @@ public:
 		return false;
 	}
 
+	virtual bool IsReductionActive(const struct FMeshReductionSettings& ReductionSettings, uint32 NumVertices, uint32 NumTriangles) const
+	{
+		return false;
+	}
+
 	virtual bool IsReductionActive(const FSkeletalMeshOptimizationSettings &ReductionSettings) const
 	{
 		return IsReductionActive(ReductionSettings, 0, 0);
diff --git a/Engine/Plugins/Interchange/Runtime/Source/Import/Private/Mesh/InterchangeStaticMeshFactory.cpp b/Engine/Plugins/Interchange/Runtime/Source/Import/Private/Mesh/InterchangeStaticMeshFactory.cpp
--- a/Engine/Plugins/Interchange/Runtime/Source/Import/Private/Mesh/InterchangeStaticMeshFactory.cpp
+++ b/Engine/Plugins/Interchange/Runtime/Source/Import/Private/Mesh/InterchangeStaticMeshFactory.cpp
@@ -490,9 +490,7 @@ void UInterchangeStaticMeshFactory::SetupSourceModelsSettings(UStaticMesh& Stati
 		// Make sure that mesh descriptions for added LODs are kept as is when the mesh is built
 		if (LodIndex >= PreviousLodCount)
 		{
-			SrcModel.ReductionSettings.MaxDeviation = 0.0f;
-			SrcModel.ReductionSettings.PercentTriangles = 1.0f;
-			SrcModel.ReductionSettings.PercentVertices = 1.0f;
+			SrcModel.ResetReductionSetting();
 		}
 
 		if (!bIsAReimport && LodMeshDescriptions.IsValidIndex(LodIndex))
diff --git a/Engine/Plugins/Runtime/MeshModelingToolset/Source/ModelingComponentsEditorOnly/Private/ToolTargets/StaticMeshToolTarget.cpp b/Engine/Plugins/Runtime/MeshModelingToolset/Source/ModelingComponentsEditorOnly/Private/ToolTargets/StaticMeshToolTarget.cpp
--- a/Engine/Plugins/Runtime/MeshModelingToolset/Source/ModelingComponentsEditorOnly/Private/ToolTargets/StaticMeshToolTarget.cpp
+++ b/Engine/Plugins/Runtime/MeshModelingToolset/Source/ModelingComponentsEditorOnly/Private/ToolTargets/StaticMeshToolTarget.cpp
@@ -311,8 +311,7 @@ void UStaticMeshToolTarget::CommitMeshDescription(UStaticMesh* StaticMeshIn, con
 
 		// configure build settings to prevent the standard static mesh reduction from running and replacing the render LOD.
 		FStaticMeshSourceModel& ThisSourceModel = StaticMeshIn->GetSourceModel((int32)EditingLODIn);
-		ThisSourceModel.ReductionSettings.PercentTriangles = 1.f;
-		ThisSourceModel.ReductionSettings.PercentVertices = 1.f;
+		ThisSourceModel.ResetReductionSetting();
 	}
 
 	StaticMeshIn->PostEditChange();
diff --git a/Engine/Source/Developer/MeshBuilder/Private/StaticMeshBuilder.cpp b/Engine/Source/Developer/MeshBuilder/Private/StaticMeshBuilder.cpp
--- a/Engine/Source/Developer/MeshBuilder/Private/StaticMeshBuilder.cpp
+++ b/Engine/Source/Developer/MeshBuilder/Private/StaticMeshBuilder.cpp
@@ -407,6 +407,8 @@ bool FStaticMeshBuilder::Build(FStaticMeshRenderData& StaticMeshRenderData, USta
 		if (bIsMeshDescriptionValid)
 		{
 			MeshDescriptionHelper.SetupRenderMeshDescription(StaticMesh, MeshDescriptions[LodIndex], false, true);
+			//Get back the reduction status once we apply all build settings, vertex count can change depending on the build settings
+			bUseReduction = StaticMesh->IsReductionActive(LodIndex, MeshDescriptions[LodIndex]);
 		}
 		else
 		{
diff --git a/Engine/Source/Developer/MeshReductionInterface/Public/IMeshReductionInterfaces.h b/Engine/Source/Developer/MeshReductionInterface/Public/IMeshReductionInterfaces.h
--- a/Engine/Source/Developer/MeshReductionInterface/Public/IMeshReductionInterfaces.h
+++ b/Engine/Source/Developer/MeshReductionInterface/Public/IMeshReductionInterfaces.h
@@ -56,6 +56,8 @@ public:
 	*	Returns true if mesh reduction is active. Active mean there will be a reduction of the vertices or triangle number
 	*/
 	virtual bool IsReductionActive(const struct FMeshReductionSettings &ReductionSettings) const = 0;
+	virtual bool IsReductionActive(const struct FMeshReductionSettings& ReductionSettings, uint32 NumVertices, uint32 NumTriangles) const = 0;
+
 	virtual bool IsReductionActive(const struct FSkeletalMeshOptimizationSettings &ReductionSettings) const = 0;
 	virtual bool IsReductionActive(const struct FSkeletalMeshOptimizationSettings &ReductionSettings, uint32 NumVertices, uint32 NumTriangles) const = 0;
 };
diff --git a/Engine/Source/Developer/MeshSimplifier/Private/QuadricMeshReduction.cpp b/Engine/Source/Developer/MeshSimplifier/Private/QuadricMeshReduction.cpp
--- a/Engine/Source/Developer/MeshSimplifier/Private/QuadricMeshReduction.cpp
+++ b/Engine/Source/Developer/MeshSimplifier/Private/QuadricMeshReduction.cpp
@@ -236,6 +236,28 @@ public:
 		uint32 NumIndexes = Indexes.Num();
 		uint32 NumTris = NumIndexes / 3;
 
+		//Get the targets tris and verts from the reduction criterion
+		uint32 TargetNumTris = NumTris;;
+		uint32 TargetNumVerts = NumVerts;
+		if (ReductionSettings.TerminationCriterion == EStaticMeshReductionTerimationCriterion::Triangles
+			|| ReductionSettings.TerminationCriterion == EStaticMeshReductionTerimationCriterion::Any)
+		{
+			TargetNumTris = FMath::CeilToInt(NumTris * ReductionSettings.PercentTriangles);
+			TargetNumTris = FMath::Min(ReductionSettings.MaxNumOfTriangles, TargetNumTris);
+		}
+
+		if (ReductionSettings.TerminationCriterion == EStaticMeshReductionTerimationCriterion::Vertices
+			|| ReductionSettings.TerminationCriterion == EStaticMeshReductionTerimationCriterion::Any)
+		{
+			TargetNumVerts = FMath::CeilToInt(NumVerts * ReductionSettings.PercentVertices);
+			TargetNumVerts = FMath::Min(ReductionSettings.MaxNumOfVerts, TargetNumVerts);
+		}
+
+		// We need a minimum of 2 triangles, to see the object on both side. If we use one, we will end up with zero triangle when we will remove a shared edge
+		TargetNumTris = FMath::Max(TargetNumTris, 2u);
+		// Clamp to a minimum of 4 vertices, ReductionSettings.PercentVertices can be zero which makes TargetNumVerts also zero
+		TargetNumVerts = FMath::Max(TargetNumVerts, 4u);
+
 		if (bUseOldMeshSimplifier)
 		{
 			static_assert(NumTexCoords == 8, "NumTexCoords changed, fix AttributeWeights");
@@ -320,12 +342,7 @@ public:
 			//MeshSimp->SetBoundaryLocked();
 			MeshSimp->InitCosts();
 
-			//We need a minimum of 2 triangles, to see the object on both side. If we use one, we will end up with zero triangle when we will remove a shared edge
-			int32 AbsoluteMinTris = 2;
-			int32 TargetNumTriangles = (ReductionSettings.TerminationCriterion != EStaticMeshReductionTerimationCriterion::Vertices) ? FMath::Max(AbsoluteMinTris, FMath::CeilToInt(NumTris * ReductionSettings.PercentTriangles)) : AbsoluteMinTris;
-			int32 TargetNumVertices = (ReductionSettings.TerminationCriterion != EStaticMeshReductionTerimationCriterion::Triangles) ? FMath::CeilToInt(NumVerts * ReductionSettings.PercentVertices) : 0;
-
-			float MaxErrorSqr = MeshSimp->SimplifyMesh(MAX_FLT, TargetNumTriangles, TargetNumVertices);
+			float MaxErrorSqr = MeshSimp->SimplifyMesh(MAX_FLT, TargetNumTris, TargetNumVerts);
 
 			MeshSimp->OutputMesh(Verts.GetData(), Indexes.GetData(), &NumVerts, &NumIndexes);
 			MeshSimp->CompactFaceData(MaterialIndexes);
@@ -336,14 +353,6 @@ public:
 		}
 		else
 		{
-			uint32 TargetNumTris = FMath::CeilToInt( NumTris * ReductionSettings.PercentTriangles );
-			uint32 TargetNumVerts = FMath::CeilToInt( NumVerts * ReductionSettings.PercentVertices );
-
-			// We need a minimum of 2 triangles, to see the object on both side. If we use one, we will end up with zero triangle when we will remove a shared edge
-			TargetNumTris = FMath::Max( TargetNumTris, 2u );
-			// Clamp to a minimum of 4 vertices, ReductionSettings.PercentVertices can be zero which makes TargetNumVerts also zero
-			TargetNumVerts = FMath::Max(TargetNumVerts, 4u);
-
 			if( TargetNumVerts < NumVerts || TargetNumTris < NumTris )
 			{
 				using VertType = TVertSimp< NumTexCoords >;
@@ -573,25 +582,34 @@ public:
 	*/
 	virtual bool IsReductionActive(const struct FMeshReductionSettings &ReductionSettings) const
 	{
-		float Threshold_One = (1.0f - KINDA_SMALL_NUMBER);
+		return IsReductionActive(ReductionSettings, 0, 0);
+	}
+
+	virtual bool IsReductionActive(const struct FMeshReductionSettings& ReductionSettings, uint32 NumVertices, uint32 NumTriangles) const
+	{
+		float Threshold_One = (1.0f - UE_KINDA_SMALL_NUMBER);
 		switch (ReductionSettings.TerminationCriterion)
 		{
-			case EStaticMeshReductionTerimationCriterion::Triangles:
-			{
-				return ReductionSettings.PercentTriangles < Threshold_One;
-			}
-			break;
-			case EStaticMeshReductionTerimationCriterion::Vertices:
-			{
-				return ReductionSettings.PercentVertices < Threshold_One;
-			}
-			break;
-			case EStaticMeshReductionTerimationCriterion::Any:
-			{
-				return ReductionSettings.PercentTriangles < Threshold_One || ReductionSettings.PercentVertices < Threshold_One;
-			}
-			break;
+		case EStaticMeshReductionTerimationCriterion::Triangles:
+		{
+			return ReductionSettings.PercentTriangles < Threshold_One || ReductionSettings.MaxNumOfTriangles < NumTriangles;
 		}
+		break;
+		case EStaticMeshReductionTerimationCriterion::Vertices:
+		{
+			return ReductionSettings.PercentVertices < Threshold_One || ReductionSettings.MaxNumOfVerts < NumVertices;
+		}
+		break;
+		case EStaticMeshReductionTerimationCriterion::Any:
+		{
+			return ReductionSettings.PercentTriangles < Threshold_One
+				|| ReductionSettings.PercentVertices < Threshold_One
+				|| ReductionSettings.MaxNumOfTriangles < NumTriangles
+				|| ReductionSettings.MaxNumOfVerts < NumVertices;
+		}
+		break;
+		}
+
 		return false;
 	}
 
diff --git a/Engine/Source/Editor/StaticMeshEditor/Private/StaticMeshEditorSubsystem.cpp b/Engine/Source/Editor/StaticMeshEditor/Private/StaticMeshEditorSubsystem.cpp
--- a/Engine/Source/Editor/StaticMeshEditor/Private/StaticMeshEditorSubsystem.cpp
+++ b/Engine/Source/Editor/StaticMeshEditor/Private/StaticMeshEditorSubsystem.cpp
@@ -781,23 +781,7 @@ int32 UStaticMeshEditorSubsystem::SetLodFromStaticMesh(UStaticMesh* DestinationS
 	// Base the reduction on the new lod
 	DestinationMeshSourceModel.ReductionSettings.BaseLODModel = DestinationLodIndex;
 
-	// Fragile. If a public function emerge to determine if a reduction will be used please consider using it and remove this code.
-	bool bDoesSourceLodUseReduction = false;
-	switch (SourceMeshSourceModel.ReductionSettings.TerminationCriterion)
-	{
-	case EStaticMeshReductionTerimationCriterion::Triangles:
-		bDoesSourceLodUseReduction = !FMath::IsNearlyEqual(SourceMeshSourceModel.ReductionSettings.PercentTriangles, 100.f);
-		break;
-	case EStaticMeshReductionTerimationCriterion::Vertices:
-		bDoesSourceLodUseReduction = !FMath::IsNearlyEqual(SourceMeshSourceModel.ReductionSettings.PercentVertices, 100.f);
-		break;
-	case EStaticMeshReductionTerimationCriterion::Any:
-		bDoesSourceLodUseReduction = !(FMath::IsNearlyEqual(SourceMeshSourceModel.ReductionSettings.PercentTriangles, 100.f) && FMath::IsNearlyEqual(SourceMeshSourceModel.ReductionSettings.PercentVertices, 100.f));
-		break;
-	default:
-		break;
-	}
-	bDoesSourceLodUseReduction |= SourceMeshSourceModel.ReductionSettings.MaxDeviation > 0.f;
+	bool bDoesSourceLodUseReduction = SourceStaticMesh->IsReductionActive(SourceLodIndex);
 
 
 	int32 BaseSourceLodIndex = bDoesSourceLodUseReduction ? SourceMeshSourceModel.ReductionSettings.BaseLODModel : SourceLodIndex;
diff --git a/Engine/Source/Editor/StaticMeshEditor/Private/StaticMeshEditorTools.cpp b/Engine/Source/Editor/StaticMeshEditor/Private/StaticMeshEditorTools.cpp
--- a/Engine/Source/Editor/StaticMeshEditor/Private/StaticMeshEditorTools.cpp
+++ b/Engine/Source/Editor/StaticMeshEditor/Private/StaticMeshEditorTools.cpp
@@ -1252,7 +1252,9 @@ bool FMeshReductionSettingsLayout::UseNativeToolLayout() const
 EVisibility FMeshReductionSettingsLayout::GetTriangleCriterionVisibility() const
 {
 	EVisibility VisibilityValue;
-	if (!bUseQuadricSimplifier || ReductionSettings.TerminationCriterion != EStaticMeshReductionTerimationCriterion::Vertices)
+	if (!bUseQuadricSimplifier
+		|| ReductionSettings.TerminationCriterion == EStaticMeshReductionTerimationCriterion::Triangles
+		|| ReductionSettings.TerminationCriterion == EStaticMeshReductionTerimationCriterion::Any)
 	{
 		VisibilityValue =  EVisibility::Visible;
 	}
@@ -1268,7 +1270,9 @@ EVisibility FMeshReductionSettingsLayout::GetTriangleCriterionVisibility() const
 EVisibility FMeshReductionSettingsLayout::GetVertexCriterionVisibility() const
 {
 	EVisibility VisibilityValue;
-	if (!bUseQuadricSimplifier || ReductionSettings.TerminationCriterion != EStaticMeshReductionTerimationCriterion::Triangles)
+	if (!bUseQuadricSimplifier
+		|| ReductionSettings.TerminationCriterion == EStaticMeshReductionTerimationCriterion::Vertices
+		|| ReductionSettings.TerminationCriterion == EStaticMeshReductionTerimationCriterion::Any)
 	{
 		VisibilityValue = EVisibility::Visible;
 	}
@@ -1327,30 +1331,76 @@ void FMeshReductionSettingsLayout::GenerateChildContent( IDetailChildrenBuilder&
 			.OnValueCommitted(this, &FMeshReductionSettingsLayout::OnPercentTrianglesCommitted)
 		]
 		.Visibility(TAttribute<EVisibility>(this, &FMeshReductionSettingsLayout::GetTriangleCriterionVisibility));
-
+		
+		if (bUseQuadricSimplifier)
+		{
+			ChildrenBuilder.AddCustomRow( LOCTEXT("MaxTrianglesPercentage_Row", "Max Number of Triangles") )
+			.RowTag("MaxPercentTriangles")
+			.NameContent()
+			[
+				SNew(STextBlock)
+				.Font( IDetailLayoutBuilder::GetDetailFont() )
+				.Text(LOCTEXT("MaxTrianglesPercentage", "Max Triangles Count"))
+				.ToolTipText(LOCTEXT("MaxTrianglesPercentage_ToolTip", "The maximum number of triangles to retain when using percentage criterion."))
+			]
+			.ValueContent()
+			[
+				SNew(SSpinBox<uint32>)
+				.Font( IDetailLayoutBuilder::GetDetailFont() )
+				.MinValue(2)
+				.MaxValue(MAX_uint32)
+				.Value(this, &FMeshReductionSettingsLayout::GetMaxNumOfPercentTriangles)
+				.OnValueChanged(this, &FMeshReductionSettingsLayout::OnMaxNumOfPercentTrianglesChanged)
+				.OnValueCommitted(this, &FMeshReductionSettingsLayout::OnMaxNumOfPercentTrianglesCommitted)
+			]
+			.Visibility(TAttribute<EVisibility>(this, &FMeshReductionSettingsLayout::GetTriangleCriterionVisibility));
+		}
 	}
 
 	if (bUseQuadricSimplifier)
 	{
+		//Percent vertices
 		ChildrenBuilder.AddCustomRow(LOCTEXT("PercentVertices", "Percent Vertices"))
 			.NameContent()
 			[
 				SNew(STextBlock)
 				.Font(IDetailLayoutBuilder::GetDetailFont())
-			.Text(LOCTEXT("PercentVertices", "Percent Vertices"))
+				.Text(LOCTEXT("PercentVertices", "Percent Vertices"))
 			]
-		.ValueContent()
+			.ValueContent()
 			[
 				SNew(SSpinBox<float>)
 				.Font(IDetailLayoutBuilder::GetDetailFont())
-			.MinValue(0.0f)
-			.MaxValue(100.0f)
-			.Value(this, &FMeshReductionSettingsLayout::GetPercentVertices)
-			.OnValueChanged(this, &FMeshReductionSettingsLayout::OnPercentVerticesChanged)
-			.OnValueCommitted(this, &FMeshReductionSettingsLayout::OnPercentVerticesCommitted)
+				.MinValue(0.0f)
+				.MaxValue(100.0f)
+				.Value(this, &FMeshReductionSettingsLayout::GetPercentVertices)
+				.OnValueChanged(this, &FMeshReductionSettingsLayout::OnPercentVerticesChanged)
+				.OnValueCommitted(this, &FMeshReductionSettingsLayout::OnPercentVerticesCommitted)
 			]
 		.Visibility(TAttribute<EVisibility>(this, &FMeshReductionSettingsLayout::GetVertexCriterionVisibility));
 
+		//Max percent absolute vertices
+		ChildrenBuilder.AddCustomRow( LOCTEXT("MaxVerticesPercentage_Row", "Max Number of Vertices") )
+			.RowTag("MaxVertices")
+			.NameContent()
+			[
+				SNew(STextBlock)
+				.Font( IDetailLayoutBuilder::GetDetailFont() )
+				.Text(LOCTEXT("MaxVerticesPercentage", "Max Vertices Count"))
+				.ToolTipText(LOCTEXT("MaxVerticesPercentage_ToolTip", "The maximum number of Vertices to retain when using percentage criterion."))
+			]
+			.ValueContent()
+			[
+				SNew(SSpinBox<uint32>)
+				.Font( IDetailLayoutBuilder::GetDetailFont() )
+				.MinValue(2)
+				.MaxValue(MAX_uint32)
+				.Value(this, &FMeshReductionSettingsLayout::GetMaxNumOfPercentVertices)
+				.OnValueChanged(this, &FMeshReductionSettingsLayout::OnMaxNumOfPercentVerticesChanged)
+				.OnValueCommitted(this, &FMeshReductionSettingsLayout::OnMaxNumOfPercentVerticesCommitted)
+			]
+			.Visibility(TAttribute<EVisibility>(this, &FMeshReductionSettingsLayout::GetVertexCriterionVisibility));
+
 	}
 
 	// Controls that only simplygon uses.
@@ -1600,11 +1650,21 @@ float FMeshReductionSettingsLayout::GetPercentTriangles() const
 	return ReductionSettings.PercentTriangles * 100.0f; // Display fraction as percentage.
 }
 
+uint32 FMeshReductionSettingsLayout::GetMaxNumOfPercentTriangles() const
+{
+	return ReductionSettings.MaxNumOfTriangles;
+}
+
 float FMeshReductionSettingsLayout::GetPercentVertices() const
 {
 	return ReductionSettings.PercentVertices * 100.0f; // Display fraction as percentage.
 }
 
+uint32 FMeshReductionSettingsLayout::GetMaxNumOfPercentVertices() const
+{
+	return ReductionSettings.MaxNumOfVerts;
+}
+
 float FMeshReductionSettingsLayout::GetMaxDeviation() const
 {
 	return ReductionSettings.MaxDeviation;
@@ -1661,6 +1721,35 @@ void FMeshReductionSettingsLayout::OnPercentVerticesCommitted(float NewValue, ET
 	OnPercentVerticesChanged(NewValue);
 }
 
+void FMeshReductionSettingsLayout::OnMaxNumOfPercentTrianglesChanged(uint32 NewValue)
+{
+	ReductionSettings.MaxNumOfTriangles = NewValue;
+}
+
+void FMeshReductionSettingsLayout::OnMaxNumOfPercentTrianglesCommitted(uint32 NewValue, ETextCommit::Type TextCommitType)
+{
+	if (FEngineAnalytics::IsAvailable())
+	{
+		FEngineAnalytics::GetProvider().RecordEvent(TEXT("Editor.Usage.StaticMesh.ReductionSettings"), TEXT("MaxNumOfTrianglesPercentage"), FString::Printf(TEXT("%d"), NewValue));
+	}
+	OnMaxNumOfPercentTrianglesChanged(NewValue);
+}
+
+
+void FMeshReductionSettingsLayout::OnMaxNumOfPercentVerticesChanged(uint32 NewValue)
+{
+	ReductionSettings.MaxNumOfVerts = NewValue;
+}
+
+void FMeshReductionSettingsLayout::OnMaxNumOfPercentVerticesCommitted(uint32 NewValue, ETextCommit::Type TextCommitType)
+{
+	if (FEngineAnalytics::IsAvailable())
+	{
+		FEngineAnalytics::GetProvider().RecordEvent(TEXT("Editor.Usage.StaticMesh.ReductionSettings"), TEXT("MaxNumOfVertsPercentage"), FString::Printf(TEXT("%d"), NewValue));
+	}
+	OnMaxNumOfPercentVerticesChanged(NewValue);
+}
+
 void FMeshReductionSettingsLayout::OnMaxDeviationChanged(float NewValue)
 {
 	ReductionSettings.MaxDeviation = NewValue;
diff --git a/Engine/Source/Editor/StaticMeshEditor/Private/StaticMeshEditorTools.h b/Engine/Source/Editor/StaticMeshEditor/Private/StaticMeshEditorTools.h
--- a/Engine/Source/Editor/StaticMeshEditor/Private/StaticMeshEditorTools.h
+++ b/Engine/Source/Editor/StaticMeshEditor/Private/StaticMeshEditorTools.h
@@ -261,9 +261,11 @@ private:
 
 	// used by native tool and simplygon
 	float GetPercentTriangles() const;
+	uint32 GetMaxNumOfPercentTriangles() const;
 
 	// used by native quadric simplifier
 	float GetPercentVertices() const;
+	uint32 GetMaxNumOfPercentVertices() const;
 
 	// used by simplygon only
 	float GetMaxDeviation() const;
@@ -275,10 +277,15 @@ private:
 	// used by native tool and simplygon
 	void OnPercentTrianglesChanged(float NewValue);
 	void OnPercentTrianglesCommitted(float NewValue, ETextCommit::Type TextCommitType);
+	void OnMaxNumOfPercentTrianglesChanged(uint32 NewValue);
+	void OnMaxNumOfPercentTrianglesCommitted(uint32 NewValue, ETextCommit::Type TextCommitType);
 
 	// Used by native code only
 	void OnPercentVerticesChanged(float NewValue);
 	void OnPercentVerticesCommitted(float NewValue, ETextCommit::Type TextCommitType);
+	void OnMaxNumOfPercentVerticesChanged(uint32 NewValue);
+	void OnMaxNumOfPercentVerticesCommitted(uint32 NewValue, ETextCommit::Type TextCommitType);
+	
 
 	//used by simplygon only
 	void OnMaxDeviationChanged(float NewValue);
diff --git a/Engine/Source/Editor/UnrealEd/Private/Fbx/FbxStaticMeshImport.cpp b/Engine/Source/Editor/UnrealEd/Private/Fbx/FbxStaticMeshImport.cpp
--- a/Engine/Source/Editor/UnrealEd/Private/Fbx/FbxStaticMeshImport.cpp
+++ b/Engine/Source/Editor/UnrealEd/Private/Fbx/FbxStaticMeshImport.cpp
@@ -1776,9 +1776,7 @@ UStaticMesh* UnFbx::FFbxImporter::ImportStaticMeshAsSingle(UObject* InParent, TA
 		//Make sure an imported mesh do not get reduce if there was no mesh data before reimport.
 		//In this case we have a generated LOD convert to a custom LOD
 		FStaticMeshSourceModel& SrcModel = StaticMesh->GetSourceModel(LODIndex);
-		SrcModel.ReductionSettings.MaxDeviation = 0.0f;
-		SrcModel.ReductionSettings.PercentTriangles = 1.0f;
-		SrcModel.ReductionSettings.PercentVertices = 1.0f;
+		SrcModel.ResetReductionSetting();
 	}
 	else if (InStaticMesh != NULL && LODIndex > 0)
 	{
diff --git a/Engine/Source/Editor/UnrealEd/Private/ImportUtils/StaticMeshImportUtils.cpp b/Engine/Source/Editor/UnrealEd/Private/ImportUtils/StaticMeshImportUtils.cpp
--- a/Engine/Source/Editor/UnrealEd/Private/ImportUtils/StaticMeshImportUtils.cpp
+++ b/Engine/Source/Editor/UnrealEd/Private/ImportUtils/StaticMeshImportUtils.cpp
@@ -726,11 +726,15 @@ TSharedPtr<FExistingStaticMeshData> StaticMeshImportUtils::SaveExistingStaticMes
 		FExistingLODMeshData& ExistingLODData = ExistingMeshDataPtr->ExistingLODData[SourceModelIndex];
 		ExistingLODData.ExistingBuildSettings = SourceModel.BuildSettings;
 		ExistingLODData.ExistingReductionSettings = SourceModel.ReductionSettings;
+		ExistingLODData.ExisitingMeshTrianglesCount = SourceModel.CacheMeshDescriptionTrianglesCount;
+		ExistingLODData.ExisitingMeshVerticesCount = SourceModel.CacheMeshDescriptionVerticesCount;
 		if (bIsReimportCustomLODOverGeneratedLOD && (SourceModelIndex == LodIndex))
 		{
 			//Reset the reduction
 			ExistingLODData.ExistingReductionSettings.PercentTriangles = 1.0f;
+			ExistingLODData.ExistingReductionSettings.MaxNumOfTriangles = MAX_uint32;
 			ExistingLODData.ExistingReductionSettings.PercentVertices = 1.0f;
+			ExistingLODData.ExistingReductionSettings.MaxNumOfVerts = MAX_uint32;
 			ExistingLODData.ExistingReductionSettings.MaxDeviation = 0.0f;
 		}
 		ExistingLODData.ExistingScreenSize = SourceModel.ScreenSize;
@@ -802,8 +806,9 @@ TSharedPtr<FExistingStaticMeshData> StaticMeshImportUtils::SaveExistingStaticMes
 }
 
 // Helper to find if some reduction settings are active
-bool IsReductionActive(const FMeshReductionSettings& ReductionSettings)
+bool IsReductionActive(const FExistingLODMeshData& ExisitingLodMeshData)
 {
+	const FMeshReductionSettings& ReductionSettings = ExisitingLodMeshData.ExistingReductionSettings;
 	bool bUseQuadricSimplier = true;
 	{
 		// Are we using our tool, or simplygon?  The tool is only changed during editor restarts
@@ -813,9 +818,28 @@ bool IsReductionActive(const FMeshReductionSettings& ReductionSettings)
 		VersionString.ParseIntoArray(SplitVersionString, TEXT("_"), true);
 		bUseQuadricSimplier = SplitVersionString[0].Equals("QuadricMeshReduction");
 	}
-	const bool bVertTermination = (bUseQuadricSimplier) && (ReductionSettings.TerminationCriterion != EStaticMeshReductionTerimationCriterion::Triangles) && (ReductionSettings.PercentVertices < 1.0f);
-	const bool bTriTermination = ReductionSettings.TerminationCriterion != EStaticMeshReductionTerimationCriterion::Vertices && (ReductionSettings.PercentTriangles < 1.0f);
-	return bTriTermination || bVertTermination || (ReductionSettings.MaxDeviation > 0.0f);
+	if (!bUseQuadricSimplier)
+	{
+		return (ReductionSettings.MaxDeviation > 0.0f)
+			|| (ReductionSettings.TerminationCriterion == EStaticMeshReductionTerimationCriterion::Triangles && (ReductionSettings.PercentTriangles < 1.0f));
+	}
+	
+	switch (ReductionSettings.TerminationCriterion)
+	{
+	case EStaticMeshReductionTerimationCriterion::Triangles:
+		return (ReductionSettings.PercentTriangles < 1.0f) || ReductionSettings.MaxNumOfTriangles < ExisitingLodMeshData.ExisitingMeshTrianglesCount;
+		break;
+	case EStaticMeshReductionTerimationCriterion::Vertices:
+		return (ReductionSettings.PercentVertices < 1.0f) || ReductionSettings.MaxNumOfVerts < ExisitingLodMeshData.ExisitingMeshVerticesCount;
+		break;
+	case EStaticMeshReductionTerimationCriterion::Any:
+		return (ReductionSettings.PercentTriangles < 1.0f)
+			|| ReductionSettings.MaxNumOfTriangles < ExisitingLodMeshData.ExisitingMeshTrianglesCount
+			|| (ReductionSettings.PercentVertices < 1.0f)
+			|| ReductionSettings.MaxNumOfVerts < ExisitingLodMeshData.ExisitingMeshVerticesCount;
+		break;
+	}
+	return false;
 }
 
 /* This function is call before building the mesh when we do a re-import*/
@@ -853,7 +877,7 @@ void StaticMeshImportUtils::RestoreExistingMeshSettings(const FExistingStaticMes
 			}
 			FMeshDescription* LODMeshDescription = NewMesh->GetMeshDescription(i);
 			bool bSwapFromGeneratedToImported = !ExistingMesh->ExistingLODData[i].ExistingMeshDescription.IsValid() && (LODMeshDescription && LODMeshDescription->Polygons().Num() > 0);
-			bool bWasReduced = IsReductionActive(ExistingMesh->ExistingLODData[i].ExistingReductionSettings);
+			bool bWasReduced = IsReductionActive(ExistingMesh->ExistingLODData[i]);
 
 			FStaticMeshSourceModel& SourceModel = NewMesh->GetSourceModel(i);
 			if (!bSwapFromGeneratedToImported && bWasReduced)
@@ -872,7 +896,7 @@ void StaticMeshImportUtils::RestoreExistingMeshSettings(const FExistingStaticMes
 		{
 			FMeshDescription* LODMeshDescription = NewMesh->GetMeshDescription(LODIndex);
 			bool bSwapFromGeneratedToImported = !ExistingMesh->ExistingLODData[LODIndex].ExistingMeshDescription.IsValid() && (LODMeshDescription && LODMeshDescription->Polygons().Num() > 0);
-			bool bWasReduced = IsReductionActive(ExistingMesh->ExistingLODData[LODIndex].ExistingReductionSettings);
+			bool bWasReduced = IsReductionActive(ExistingMesh->ExistingLODData[LODIndex]);
 
 			FStaticMeshSourceModel& SourceModel = NewMesh->GetSourceModel(LODIndex);
 			if (!bSwapFromGeneratedToImported && bWasReduced)
@@ -1091,7 +1115,7 @@ void StaticMeshImportUtils::RestoreExistingMeshData(const TSharedPtr<const FExis
 		FMeshDescription* LODMeshDescription = NewMesh->GetMeshDescription(i);
 		//Restore the reduction settings only if the existing data was a using reduction. Because we can set some value if we reimport from existing rawmesh to auto generated.
 		bool bSwapFromGeneratedToImported = !ExistingMeshDataPtr->ExistingLODData[i].ExistingMeshDescription.IsValid() && (LODMeshDescription && LODMeshDescription->Polygons().Num() > 0);
-		bool bWasReduced = IsReductionActive(ExistingMeshDataPtr->ExistingLODData[i].ExistingReductionSettings);
+		bool bWasReduced = IsReductionActive(ExistingMeshDataPtr->ExistingLODData[i]);
 		if ( !bSwapFromGeneratedToImported && bWasReduced)
 		{
 			SourceModel.ReductionSettings = ExistingMeshDataPtr->ExistingLODData[i].ExistingReductionSettings;
@@ -1132,7 +1156,7 @@ void StaticMeshImportUtils::RestoreExistingMeshData(const TSharedPtr<const FExis
 
 			
 			//When re-importing the asset, do not touch the LOD that was imported from file, the material array is keep intact so the section should still be valid.
-			bool NoRemapForThisLOD = LodLevel == INDEX_NONE && i != 0 && !NewMesh->GetSourceModel(i).bImportWithBaseMesh && !IsReductionActive(NewMesh->GetSourceModel(i).ReductionSettings);
+			bool NoRemapForThisLOD = LodLevel == INDEX_NONE && i != 0 && !NewMesh->GetSourceModel(i).bImportWithBaseMesh && !NewMesh->IsReductionActive(i);
 
 			FStaticMeshLODResources& LOD = NewMesh->GetRenderData()->LODResources[i];
 			
diff --git a/Engine/Source/Editor/UnrealEd/Public/ImportUtils/StaticMeshImportUtils.h b/Engine/Source/Editor/UnrealEd/Public/ImportUtils/StaticMeshImportUtils.h
--- a/Engine/Source/Editor/UnrealEd/Public/ImportUtils/StaticMeshImportUtils.h
+++ b/Engine/Source/Editor/UnrealEd/Public/ImportUtils/StaticMeshImportUtils.h
@@ -32,6 +32,8 @@ struct FExistingLODMeshData
 {
 	FMeshBuildSettings				ExistingBuildSettings;
 	FMeshReductionSettings			ExistingReductionSettings;
+	uint32							ExisitingMeshTrianglesCount;
+	uint32							ExisitingMeshVerticesCount;
 	TUniquePtr<FMeshDescription>	ExistingMeshDescription;
 	TArray<FStaticMaterial>			ExistingMaterials;
 	FPerPlatformFloat				ExistingScreenSize;
diff --git a/Engine/Source/Runtime/Engine/Classes/Engine/MeshMerging.h b/Engine/Source/Runtime/Engine/Classes/Engine/MeshMerging.h
--- a/Engine/Source/Runtime/Engine/Classes/Engine/MeshMerging.h
+++ b/Engine/Source/Runtime/Engine/Classes/Engine/MeshMerging.h
@@ -30,9 +30,9 @@ namespace EMeshFeatureImportance
 UENUM()
 enum class EStaticMeshReductionTerimationCriterion : uint8
 {
-	Triangles,
-	Vertices,
-	Any
+	Triangles UMETA(DisplayName = "Triangles", ToolTip = "Triangle percent criterion will be used for simplification."),
+	Vertices UMETA(DisplayName = "Vertice", ToolTip = "Vertice percent criterion will be used for simplification."),
+	Any UMETA(DisplayName = "First Percent Satisfied", ToolTip = "Simplification will continue until either Triangle or Vertex count criteria is met."),
 };
 
 /** Settings used to reduce a mesh. */
@@ -41,14 +41,22 @@ struct FMeshReductionSettings
 {
 	GENERATED_USTRUCT_BODY()
 
-	/** Percentage of triangles to keep. 1.0 = no reduction, 0.0 = no triangles. */
+	/** Percentage of triangles to keep. 1.0 = no reduction, 0.0 = no triangles. (Triangles criterion properties) */
 	UPROPERTY(BlueprintReadWrite, EditAnywhere, Category = ReductionSettings)
 	float PercentTriangles;
 
-	/** Percentage of vertices to keep. 1.0 = no reduction, 0.0 = no vertices. */
+	/** The maximum number of triangles to retain when using percentage termination criterion. (Triangles criterion properties) */
+	UPROPERTY(EditAnywhere, Category = ReductionMethod, meta = (DisplayName = "Max Triangle Count", ClampMin = 2, UIMin = "2"))
+	uint32 MaxNumOfTriangles;
+
+	/** Percentage of vertices to keep. 1.0 = no reduction, 0.0 = no vertices. (Vertices criterion properties) */
 	UPROPERTY(BlueprintReadWrite, EditAnywhere, Category = ReductionSettings)
 	float PercentVertices;
 
+	/** The maximum number of vertices to retain when using percentage termination criterion. (Vertices criterion properties) */
+	UPROPERTY(EditAnywhere, Category = ReductionMethod, meta = (DisplayName = "Max Vertex Count", ClampMin = 4, UIMin = "4"))
+	uint32 MaxNumOfVerts;
+
 	/** The maximum distance in object space by which the reduced mesh may deviate from the original mesh. */
 	UPROPERTY(BlueprintReadWrite, EditAnywhere, Category = ReductionSettings)
 	float MaxDeviation;
@@ -110,7 +118,9 @@ struct FMeshReductionSettings
 	/** Default settings. */
 	FMeshReductionSettings()
 		: PercentTriangles(1.0f)
+		, MaxNumOfTriangles(MAX_uint32)
 		, PercentVertices(1.0f)
+		, MaxNumOfVerts(MAX_uint32)
 		, MaxDeviation(0.0f)
 		, PixelError(8.0f)
 		, WeldingThreshold(0.0f)
@@ -132,7 +142,9 @@ struct FMeshReductionSettings
 
 	FMeshReductionSettings(const FMeshReductionSettings& Other)
 		: PercentTriangles(Other.PercentTriangles)
+		, MaxNumOfTriangles(Other.MaxNumOfTriangles)
 		, PercentVertices(Other.PercentVertices)
+		, MaxNumOfVerts(Other.MaxNumOfVerts)
 		, MaxDeviation(Other.MaxDeviation)
 		, PixelError(Other.PixelError)
 		, WeldingThreshold(Other.WeldingThreshold)
@@ -159,6 +171,8 @@ struct FMeshReductionSettings
 			TerminationCriterion == Other.TerminationCriterion
 			&& PercentVertices == Other.PercentVertices
 			&& PercentTriangles == Other.PercentTriangles
+			&& MaxNumOfTriangles == Other.MaxNumOfTriangles
+			&& MaxNumOfVerts == Other.MaxNumOfVerts
 			&& MaxDeviation == Other.MaxDeviation
 			&& PixelError == Other.PixelError
 			&& WeldingThreshold == Other.WeldingThreshold
diff --git a/Engine/Source/Runtime/Engine/Classes/Engine/StaticMesh.h b/Engine/Source/Runtime/Engine/Classes/Engine/StaticMesh.h
--- a/Engine/Source/Runtime/Engine/Classes/Engine/StaticMesh.h
+++ b/Engine/Source/Runtime/Engine/Classes/Engine/StaticMesh.h
@@ -1945,6 +1945,8 @@ public:
 	/* Return true if the reduction settings are setup to reduce a LOD*/
 	ENGINE_API bool IsReductionActive(int32 LODIndex) const;
 
+	ENGINE_API bool IsReductionActive(int32 LodIndex, const FMeshDescription& MeshDescription) const;
+
 	/* Get a copy of the reduction settings for a specified LOD index. */
 	ENGINE_API struct FMeshReductionSettings GetReductionSettings(int32 LODIndex) const;
 
diff --git a/Engine/Source/Runtime/Engine/Classes/Engine/StaticMeshSourceData.h b/Engine/Source/Runtime/Engine/Classes/Engine/StaticMeshSourceData.h
--- a/Engine/Source/Runtime/Engine/Classes/Engine/StaticMeshSourceData.h
+++ b/Engine/Source/Runtime/Engine/Classes/Engine/StaticMeshSourceData.h
@@ -146,6 +146,7 @@ public:
 	 */
 	ENGINE_API void ClearMeshDescription();
 
+	ENGINE_API void ResetReductionSetting();
 #endif
 
 public:
@@ -157,6 +158,12 @@ public:
 	UPROPERTY(EditAnywhere, Category=ReductionSettings)
 	FMeshReductionSettings ReductionSettings; 
 
+	UPROPERTY()
+	uint32 CacheMeshDescriptionTrianglesCount = MAX_uint32;
+	
+	UPROPERTY()
+	uint32 CacheMeshDescriptionVerticesCount = MAX_uint32;
+
 	UPROPERTY()
 	float LODDistance_DEPRECATED;
 
diff --git a/Engine/Source/Runtime/Engine/Private/StaticMesh.cpp b/Engine/Source/Runtime/Engine/Private/StaticMesh.cpp
--- a/Engine/Source/Runtime/Engine/Private/StaticMesh.cpp
+++ b/Engine/Source/Runtime/Engine/Private/StaticMesh.cpp
@@ -2376,11 +2376,43 @@ void UStaticMesh::GetLODGroupsDisplayNames(TArray<FText>& OutLODGroupsDisplayNam
 	RunningPlatform->GetStaticMeshLODSettings().GetLODGroupDisplayNames(OutLODGroupsDisplayNames);
 }
 
-bool UStaticMesh::IsReductionActive(int32 LODIndex) const
+bool UStaticMesh::IsReductionActive(int32 LodIndex) const
 {
-	FMeshReductionSettings ReductionSettings = GetReductionSettings(LODIndex);
-	IMeshReduction* ReductionModule = FModuleManager::Get().LoadModuleChecked<IMeshReductionManagerModule>("MeshReductionInterface").GetStaticMeshReductionInterface();
-	return ReductionModule->IsReductionActive(ReductionSettings);
+	//Invalid LOD are not reduced
+	if (!IsSourceModelValid(LodIndex))
+	{
+		return false;
+	}
+
+	bool bReductionActive = false;
+	if (IMeshReduction* ReductionModule = FModuleManager::Get().LoadModuleChecked<IMeshReductionManagerModule>("MeshReductionInterface").GetStaticMeshReductionInterface())
+	{
+		FMeshReductionSettings ReductionSettings = GetReductionSettings(LodIndex);
+		const FStaticMeshSourceModel& SrcModel = GetSourceModel(LodIndex);
+		uint32 LODTriNumber = SrcModel.CacheMeshDescriptionTrianglesCount;
+		uint32 LODVertexNumber = SrcModel.CacheMeshDescriptionVerticesCount;
+		bReductionActive = ReductionModule->IsReductionActive(ReductionSettings, LODVertexNumber, LODTriNumber);
+	}
+	return bReductionActive;
+}
+
+bool UStaticMesh::IsReductionActive(int32 LodIndex, const FMeshDescription& MeshDescription) const
+{
+	//Invalid LOD are not reduced
+	if (!IsSourceModelValid(LodIndex))
+	{
+		return false;
+	}
+	bool bReductionActive = false;
+	if (IMeshReduction* ReductionModule = FModuleManager::Get().LoadModuleChecked<IMeshReductionManagerModule>("MeshReductionInterface").GetStaticMeshReductionInterface())
+	{
+		FStaticMeshSourceModel* MutableSourceModel = const_cast<FStaticMeshSourceModel*>(&GetSourceModel(LodIndex));
+		MutableSourceModel->CacheMeshDescriptionVerticesCount = FStaticMeshOperations::GetUniqueVertexCount(MeshDescription);
+		MutableSourceModel->CacheMeshDescriptionTrianglesCount = MeshDescription.Triangles().Num();
+		FMeshReductionSettings ReductionSettings = GetReductionSettings(LodIndex);
+		bReductionActive = ReductionModule->IsReductionActive(ReductionSettings, MutableSourceModel->CacheMeshDescriptionVerticesCount, MutableSourceModel->CacheMeshDescriptionTrianglesCount);
+	}
+	return bReductionActive;
 }
 
 FMeshReductionSettings UStaticMesh::GetReductionSettings(int32 LODIndex) const
@@ -2470,7 +2502,9 @@ static void SerializeReductionSettingsForDDC(FArchive& Ar, FMeshReductionSetting
 	// Note: this serializer is only used to build the mesh DDC key, no versioning is required
 	Ar << ReductionSettings.TerminationCriterion;
 	Ar << ReductionSettings.PercentTriangles;
+	Ar << ReductionSettings.MaxNumOfTriangles;
 	Ar << ReductionSettings.PercentVertices;
+	Ar << ReductionSettings.MaxNumOfVerts;
 	Ar << ReductionSettings.MaxDeviation;
 	Ar << ReductionSettings.PixelError;
 	Ar << ReductionSettings.WeldingThreshold;
@@ -4140,9 +4174,7 @@ bool UStaticMesh::SetCustomLOD(const UStaticMesh* SourceStaticMesh, int32 Destin
 		//Make sure an imported mesh do not get reduce if there was no mesh data before reimport.
 		//In this case we have a generated LOD convert to a custom LOD
 		FStaticMeshSourceModel& SrcModel = GetSourceModel(DestinationLodIndex);
-		SrcModel.ReductionSettings.MaxDeviation = 0.0f;
-		SrcModel.ReductionSettings.PercentTriangles = 1.0f;
-		SrcModel.ReductionSettings.PercentVertices = 1.0f;
+		SrcModel.ResetReductionSetting();
 	}
 	else
 	{
diff --git a/Engine/Source/Runtime/Engine/Private/StaticMeshSourceData.cpp b/Engine/Source/Runtime/Engine/Private/StaticMeshSourceData.cpp
--- a/Engine/Source/Runtime/Engine/Private/StaticMeshSourceData.cpp
+++ b/Engine/Source/Runtime/Engine/Private/StaticMeshSourceData.cpp
@@ -209,6 +209,11 @@ bool FStaticMeshSourceModel::LoadMeshDescription(FMeshDescription& OutMeshDescri
 	{
 		// Unpack MeshDescription from the bulk data which was deserialized
 		StaticMeshDescriptionBulkData->GetBulkData().LoadMeshDescription(OutMeshDescription);
+		
+		//Set the tris/verts count cache in case we load an asset that was save before the cache exist
+		FStaticMeshSourceModel* MutableThis = const_cast<FStaticMeshSourceModel*>(this);
+		MutableThis->CacheMeshDescriptionTrianglesCount = static_cast<uint32>(OutMeshDescription.Triangles().Num());
+		MutableThis->CacheMeshDescriptionVerticesCount = static_cast<uint32>(FStaticMeshOperations::GetUniqueVertexCount(OutMeshDescription));
 		return true;
 	}
 
@@ -329,10 +334,20 @@ void FStaticMeshSourceModel::CommitMeshDescription(bool bUseHashAsGuid)
 	// as long as no more than one thread is calling it for the same UStaticMesh.
 
 	check(StaticMeshDescriptionBulkData != nullptr);
+
+	//Reset the mesh tris/verts count cache
+	CacheMeshDescriptionTrianglesCount = MAX_uint32;
+	CacheMeshDescriptionVerticesCount = MAX_uint32;
 	if (StaticMeshDescriptionBulkData->HasCachedMeshDescription())
 	{
 		// Package up mesh description into bulk data
 		StaticMeshDescriptionBulkData->CommitMeshDescription(bUseHashAsGuid);
+		//Set the tris/verts count cache
+		if (FMeshDescription* MeshDescription = GetCachedMeshDescription())
+		{
+			CacheMeshDescriptionTrianglesCount = static_cast<uint32>(MeshDescription->Triangles().Num());
+			CacheMeshDescriptionVerticesCount = static_cast<uint32>(FStaticMeshOperations::GetUniqueVertexCount(*MeshDescription));
+		}
 	}
 	else
 	{
@@ -350,6 +365,14 @@ void FStaticMeshSourceModel::ClearMeshDescription()
 	StaticMeshDescriptionBulkData->RemoveMeshDescription();
 }
 
+void FStaticMeshSourceModel::ResetReductionSetting()
+{
+	ReductionSettings.MaxDeviation = 0.0f;
+	ReductionSettings.PercentTriangles = 1.f;
+	ReductionSettings.MaxNumOfTriangles = MAX_uint32;
+	ReductionSettings.PercentVertices = 1.f;
+	ReductionSettings.MaxNumOfVerts = MAX_uint32;
+}
 
 void FStaticMeshSourceModel::SerializeBulkData(FArchive& Ar, UObject* Owner)
 {
diff --git a/Engine/Source/Runtime/StaticMeshDescription/Private/StaticMeshOperations.cpp b/Engine/Source/Runtime/StaticMeshDescription/Private/StaticMeshOperations.cpp
--- a/Engine/Source/Runtime/StaticMeshDescription/Private/StaticMeshOperations.cpp
+++ b/Engine/Source/Runtime/StaticMeshDescription/Private/StaticMeshOperations.cpp
@@ -2559,4 +2559,161 @@ void FStaticMeshOperations::ApplyTransform(FMeshDescription& MeshDescription, co
 	}
 }
 
+namespace UE::Private
+{
+	class FPrivateVertexInfo
+	{
+	public:
+		FVector3f			Position;
+		FVector3f			Normal;
+		FVector3f			Tangents[2];
+		FLinearColor		Color;
+		FVector2f			TexCoords[MAX_MESH_TEXTURE_COORDS_MD];
+
+		void Validate()
+		{
+			Normal.Normalize();
+			Tangents[0] -= (Tangents[0] | Normal) * Normal;
+			Tangents[0].Normalize();
+			Tangents[1] -= (Tangents[1] | Normal) * Normal;
+			Tangents[1] -= (Tangents[1] | Tangents[0]) * Tangents[0];
+			Tangents[1].Normalize();
+			Color = Color.GetClamped();
+		}
+
+		bool Equals(const FPrivateVertexInfo& Other) const
+		{
+			constexpr float UVEpsilon = 1.0f / 1024.0f;
+			if (!Position.Equals(Other.Position, UE_THRESH_POINTS_ARE_SAME) ||
+				!Tangents[0].Equals(Other.Tangents[0], UE_THRESH_NORMALS_ARE_SAME) ||
+				!Tangents[1].Equals(Other.Tangents[1], UE_THRESH_NORMALS_ARE_SAME) ||
+				!Normal.Equals(Other.Normal, UE_THRESH_NORMALS_ARE_SAME) ||
+				!Color.Equals(Other.Color))
+			{
+				return false;
+			}
+
+			// UVs
+			for (int32 UVIndex = 0; UVIndex < MAX_MESH_TEXTURE_COORDS_MD; UVIndex++)
+			{
+				if (!TexCoords[UVIndex].Equals(Other.TexCoords[UVIndex], UVEpsilon))
+				{
+					return false;
+				}
+			}
+
+			return true;
+		}
+	};
+}
+
+int32 FStaticMeshOperations::GetUniqueVertexCount(const FMeshDescription& MeshDescription)
+{
+	constexpr uint32 NumTexCoords = MAX_MESH_TEXTURE_COORDS_MD;
+	FOverlappingCorners OverlappingCorners;
+	FStaticMeshOperations::FindOverlappingCorners(OverlappingCorners, MeshDescription, UE_THRESH_POINTS_ARE_SAME);
+	TArray< UE::Private::FPrivateVertexInfo > Verts;
+	Verts.Reserve(MeshDescription.Vertices().Num());
+	TMap< int32, int32 > VertsMap;
+	int32 NumFaces = MeshDescription.Triangles().Num();
+	int32 NumWedges = NumFaces * 3;
+	const FStaticMeshConstAttributes MeshAttribute(MeshDescription);
+	TVertexAttributesConstRef<FVector3f> VertexPositions = MeshAttribute.GetVertexPositions();
+	TVertexInstanceAttributesConstRef<FVector3f> VertexNormals = MeshAttribute.GetVertexInstanceNormals();
+	TVertexInstanceAttributesConstRef<FVector3f> VertexTangents = MeshAttribute.GetVertexInstanceTangents();
+	TVertexInstanceAttributesConstRef<float> VertexBinormalSigns = MeshAttribute.GetVertexInstanceBinormalSigns();
+	TVertexInstanceAttributesConstRef<FVector4f> VertexColors = MeshAttribute.GetVertexInstanceColors();
+	TVertexInstanceAttributesConstRef<FVector2f> VertexUVs = MeshAttribute.GetVertexInstanceUVs();
+
+	int32 WedgeIndex = 0;
+	for (const FTriangleID TriangleID : MeshDescription.Triangles().GetElementIDs())
+	{
+		TArrayView<const FVertexID> VertexIDs = MeshDescription.GetTriangleVertices(TriangleID);
+
+		FVector3f CornerPositions[3];
+		for (int32 TriVert = 0; TriVert < 3; ++TriVert)
+		{
+			const FVertexID VertexID = VertexIDs[TriVert];
+			CornerPositions[TriVert] = VertexPositions[VertexID];
+		}
+
+		// Don't process degenerate triangles.
+		if (CornerPositions[0].Equals(CornerPositions[1], UE_THRESH_POINTS_ARE_SAME) ||
+			CornerPositions[0].Equals(CornerPositions[2], UE_THRESH_POINTS_ARE_SAME) ||
+			CornerPositions[1].Equals(CornerPositions[2], UE_THRESH_POINTS_ARE_SAME))
+		{
+			WedgeIndex += 3;
+			continue;
+		}
+
+		for (int32 TriVert = 0; TriVert < 3; ++TriVert, ++WedgeIndex)
+		{
+			const FVertexInstanceID VertexInstanceID = MeshDescription.GetTriangleVertexInstance(TriangleID, TriVert);
+			const FVector3f& VertexPosition = CornerPositions[TriVert];
+			UE::Private::FPrivateVertexInfo NewVert;
+			NewVert.Position = CornerPositions[TriVert];
+			NewVert.Tangents[0] = VertexTangents[VertexInstanceID];
+			NewVert.Normal = VertexNormals[VertexInstanceID];
+			NewVert.Tangents[1] = FVector3f(0.0f);
+			if (!NewVert.Normal.IsNearlyZero(SMALL_NUMBER) && !NewVert.Tangents[0].IsNearlyZero(SMALL_NUMBER))
+			{
+				NewVert.Tangents[1] = FVector3f::CrossProduct(NewVert.Normal, NewVert.Tangents[0]).GetSafeNormal() * VertexBinormalSigns[VertexInstanceID];
+			}
+
+			// Fix bad tangents
+			NewVert.Tangents[0] = NewVert.Tangents[0].ContainsNaN() ? FVector3f::ZeroVector : NewVert.Tangents[0];
+			NewVert.Tangents[1] = NewVert.Tangents[1].ContainsNaN() ? FVector3f::ZeroVector : NewVert.Tangents[1];
+			NewVert.Normal = NewVert.Normal.ContainsNaN() ? FVector3f::ZeroVector : NewVert.Normal;
+			NewVert.Color = FLinearColor(VertexColors[VertexInstanceID]);
+
+			for (int32 UVIndex = 0; UVIndex < NumTexCoords; UVIndex++)
+			{
+				if (UVIndex < VertexUVs.GetNumChannels())
+				{
+					NewVert.TexCoords[UVIndex] = VertexUVs.Get(VertexInstanceID, UVIndex);
+				}
+				else
+				{
+					NewVert.TexCoords[UVIndex] = FVector2f::ZeroVector;
+				}
+			}
+
+			// Make sure this vertex is valid from the start
+			NewVert.Validate();
+
+			//Never add duplicated vertex instance
+			//Use WedgeIndex since OverlappingCorners has been built based on that
+			const TArray<int32>& DupVerts = OverlappingCorners.FindIfOverlapping(WedgeIndex);
+
+			int32 Index = INDEX_NONE;
+			for (int32 k = 0; k < DupVerts.Num(); k++)
+			{
+				if (DupVerts[k] >= WedgeIndex)
+				{
+					// the verts beyond me haven't been placed yet, so these duplicates are not relevant
+					break;
+				}
+
+				int32* Location = VertsMap.Find(DupVerts[k]);
+				if (Location)
+				{
+					UE::Private::FPrivateVertexInfo& FoundVert = Verts[*Location];
+
+					if (NewVert.Equals(FoundVert))
+					{
+						Index = *Location;
+						break;
+					}
+				}
+			}
+			if (Index == INDEX_NONE)
+			{
+				Index = Verts.Add(NewVert);
+				VertsMap.Add(WedgeIndex, Index);
+			}
+		}
+	}
+	return Verts.Num();
+}
+
 #undef LOCTEXT_NAMESPACE
diff --git a/Engine/Source/Runtime/StaticMeshDescription/Public/StaticMeshOperations.h b/Engine/Source/Runtime/StaticMeshDescription/Public/StaticMeshOperations.h
--- a/Engine/Source/Runtime/StaticMeshDescription/Public/StaticMeshOperations.h
+++ b/Engine/Source/Runtime/StaticMeshDescription/Public/StaticMeshOperations.h
@@ -182,4 +182,10 @@ public:
 
 	/** Transforms the MeshDescription data using the provided transform. */
 	static void ApplyTransform(FMeshDescription& MeshDescription, const FTransform& Transform);
+
+	/**
+	 * Return the number of unique vertices, unique vertices are the result of welding all similar vertex instances (position, UV, tangent space, color,...)
+	 * 
+	 */
+	static int32 GetUniqueVertexCount(const FMeshDescription& MeshDescription);
 };
