Expose StaticMesh LOD ScreenSize property (a detail customization row)

diff --git a/Engine/Source/Editor/StaticMeshEditor/Private/StaticMeshEditorTools.cpp b/Engine/Source/Editor/StaticMeshEditor/Private/StaticMeshEditorTools.cpp
--- a/Engine/Source/Editor/StaticMeshEditor/Private/StaticMeshEditorTools.cpp
+++ b/Engine/Source/Editor/StaticMeshEditor/Private/StaticMeshEditorTools.cpp
@@ -3872,23 +3872,25 @@ void FLevelOfDetailSettingsLayout::AddLODLevelCategories( IDetailLayoutBuilder&
 			TAttribute<TArray<FName>> PlatformOverrideNames = TAttribute<TArray<FName>>::Create(TAttribute<TArray<FName>>::FGetter::CreateSP(this, &FLevelOfDetailSettingsLayout::GetLODScreenSizePlatformOverrideNames, LODIndex));
 
 			FPerPlatformPropertyCustomNodeBuilderArgs Args;
-			Args.FilterText = LOCTEXT("ScreenSizeName", "Screen Size");
-
-			TAttribute<bool> IsScreenSizeEnabled = TAttribute<bool>::CreateSP(this, &FLevelOfDetailSettingsLayout::CanChangeLODScreenSize);
-
-			Args.OnGenerateNameWidget = FOnGetContent::CreateLambda([IsScreenSizeEnabled]()
-				{
-					return SNew(STextBlock)
-						.IsEnabled(IsScreenSizeEnabled)
-						.Font(IDetailLayoutBuilder::GetDetailFont())
-						.Text(LOCTEXT("ScreenSizeName", "Screen Size"));
-				});
-
-			Args.PlatformOverrideNames = PlatformOverrideNames;
-			Args.OnAddPlatformOverride = FOnPlatformOverrideAction::CreateSP(this, &FLevelOfDetailSettingsLayout::AddLODScreenSizePlatformOverride, LODIndex);
-			Args.OnRemovePlatformOverride = FOnPlatformOverrideAction::CreateSP(this, &FLevelOfDetailSettingsLayout::RemoveLODScreenSizePlatformOverride, LODIndex);
-			Args.OnGenerateWidgetForPlatformRow = FOnGenerateWidget::CreateSP(this, &FLevelOfDetailSettingsLayout::GetLODScreenSizeWidget, LODIndex);
-			Args.IsEnabled = IsScreenSizeEnabled;
+			{
+				FText ScreenSizePropertyText(LOCTEXT("ScreenSizeName", "Screen Size"));
+				TAttribute<bool> IsScreenSizeEnabled = TAttribute<bool>::CreateSP(this, &FLevelOfDetailSettingsLayout::CanChangeLODScreenSize);
+
+				Args.Name = FName("ScreenSize");
+				Args.FilterText = ScreenSizePropertyText;
+				Args.PlatformOverrideNames = PlatformOverrideNames;
+				Args.OnAddPlatformOverride = FOnPlatformOverrideAction::CreateSP(this, &FLevelOfDetailSettingsLayout::AddLODScreenSizePlatformOverride, LODIndex);
+				Args.OnRemovePlatformOverride = FOnPlatformOverrideAction::CreateSP(this, &FLevelOfDetailSettingsLayout::RemoveLODScreenSizePlatformOverride, LODIndex);
+				Args.OnGenerateWidgetForPlatformRow = FOnGenerateWidget::CreateSP(this, &FLevelOfDetailSettingsLayout::GetLODScreenSizeWidget, LODIndex);
+				Args.IsEnabled = IsScreenSizeEnabled;
+				Args.OnGenerateNameWidget = FOnGetContent::CreateLambda([IsScreenSizeEnabled = MoveTemp(IsScreenSizeEnabled), ScreenSizePropertyText = MoveTemp(ScreenSizePropertyText)]()
+					{
+						return SNew(STextBlock)
+							.IsEnabled(IsScreenSizeEnabled)
+							.Font(IDetailLayoutBuilder::GetDetailFont())
+							.Text(ScreenSizePropertyText);
+					});
+			}
 
 			LODCategory.AddCustomBuilder(MakeShared<FPerPlatformPropertyCustomNodeBuilder>(MoveTemp(Args)));
 
