Adding cvar to determine the default collapse mode of Nanite Settings in the details panel of the Static Mesh Editor Tool.

#rnx
#test editor

diff --git a/Engine/Source/Editor/StaticMeshEditor/Private/StaticMeshEditorTools.cpp b/Engine/Source/Editor/StaticMeshEditor/Private/StaticMeshEditorTools.cpp
--- a/Engine/Source/Editor/StaticMeshEditor/Private/StaticMeshEditorTools.cpp
+++ b/Engine/Source/Editor/StaticMeshEditor/Private/StaticMeshEditorTools.cpp
@@ -70,6 +70,14 @@ const int32 DefaultVertsPerHull = 16;
 #define LOCTEXT_NAMESPACE "StaticMeshEditor"
 DEFINE_LOG_CATEGORY_STATIC(LogStaticMeshEditorTools,Log,All);
 
+bool GIsNaniteStaticMeshSettingsInitiallyCollapsed = 0;
+static FAutoConsoleVariableRef CVarIsNaniteStaticMeshSettingsInitiallyCollapsed(
+	TEXT("r.Nanite.IsNaniteStaticMeshSettingsInitiallyCollapsed"),
+	GIsNaniteStaticMeshSettingsInitiallyCollapsed,
+	TEXT("If the Nanite Settings are initially collapsed in the details panel in the Static Mesh Editor Tool."),
+	ECVF_ReadOnly
+);
+
 /*
 * Custom data key
 */
@@ -5258,6 +5266,7 @@ void FNaniteSettingsLayout::AddToDetailsPanel(IDetailLayoutBuilder& DetailBuilde
 	IDetailCategoryBuilder& NaniteSettingsCategory =
 		DetailBuilder.EditCategory("NaniteSettings", NaniteCategoryName);
 	NaniteSettingsCategory.SetSortOrder(10);
+	NaniteSettingsCategory.InitiallyCollapsed(GIsNaniteStaticMeshSettingsInitiallyCollapsed);
 	
 	NaniteSettingsCategory.HeaderContent
 	(
