Keep all Nanite settings controls in Static Mesh Editor enabled when Nanite is force enabled on the mesh.

diff --git a/Engine/Source/Editor/StaticMeshEditor/Private/StaticMeshEditorTools.cpp b/Engine/Source/Editor/StaticMeshEditor/Private/StaticMeshEditorTools.cpp
--- a/Engine/Source/Editor/StaticMeshEditor/Private/StaticMeshEditorTools.cpp
+++ b/Engine/Source/Editor/StaticMeshEditor/Private/StaticMeshEditorTools.cpp
@@ -5799,7 +5799,8 @@ FString FNaniteSettingsLayout::MinimumResidencyValueToDisplayString(uint32 Value
 
 ECheckBoxState FNaniteSettingsLayout::IsEnabledChecked() const
 {
-	return NaniteSettings.bEnabled ? ECheckBoxState::Checked : ECheckBoxState::Unchecked;
+	UStaticMesh* StaticMesh = StaticMeshEditor.GetStaticMesh();
+	return NaniteSettings.bEnabled || (StaticMesh && StaticMesh->IsNaniteForceEnabled()) ? ECheckBoxState::Checked : ECheckBoxState::Unchecked;
 }
 
 void FNaniteSettingsLayout::OnEnabledChanged(ECheckBoxState NewState)
