UDN bug fix:  UStaticMeshEditorSubsystem::SetLodScreenSizes should allow values greater than 1.0.  Originally it was clamped to 1.0, as the user interface sliders were clamped at that, but it's possible and useful to numerically set them outside that range, and the automation command shouldn't disallow it.

diff --git a/Engine/Source/Editor/StaticMeshEditor/Private/StaticMeshEditorSubsystem.cpp b/Engine/Source/Editor/StaticMeshEditor/Private/StaticMeshEditorSubsystem.cpp
--- a/Engine/Source/Editor/StaticMeshEditor/Private/StaticMeshEditorSubsystem.cpp
+++ b/Engine/Source/Editor/StaticMeshEditor/Private/StaticMeshEditorSubsystem.cpp
@@ -1048,7 +1048,7 @@ bool UStaticMeshEditorSubsystem::SetLodScreenSizes(UStaticMesh* StaticMesh, cons
 	StaticMesh->bAutoComputeLODScreenSize = 0;
 
 	// Arbitrarily set this to a value that won't affect the monotonic clamping on the first iteration of the loop
-	float LastScreenSize = 1.0f + 2.0f * MonotonicDifference;
+	float LastScreenSize = ScreenSizes[0] + 2.0f * MonotonicDifference;
 
 	for (int i = 0; i < StaticMesh->GetNumLODs(); i++)
 	{
@@ -1062,7 +1062,7 @@ bool UStaticMeshEditorSubsystem::SetLodScreenSizes(UStaticMesh* StaticMesh, cons
 			ScreenSizeForLOD = LastScreenSize - MonotonicDifference;
 		}
 
-		ScreenSizeForLOD = FMath::Clamp(ScreenSizeForLOD, 0.0f, 1.0f);
+		ScreenSizeForLOD = FMath::Max(ScreenSizeForLOD, 0.0f);
 
 		// Track if the input values needed to be sanitized in any way, so we can warn the user this happened.
 		if (i < ScreenSizes.Num() && ScreenSizeForLOD != ScreenSizes[i])
