Asset Editors: Add support for opening cooked assets in a new "Read Only" mode for asset editors, starting with the Static Mesh Editor and Texture Editor.

diff --git a/Engine/Plugins/Editor/EngineAssetDefinitions/Source/Private/AssetDefinition_StaticMesh.cpp b/Engine/Plugins/Editor/EngineAssetDefinitions/Source/Private/AssetDefinition_StaticMesh.cpp
--- a/Engine/Plugins/Editor/EngineAssetDefinitions/Source/Private/AssetDefinition_StaticMesh.cpp
+++ b/Engine/Plugins/Editor/EngineAssetDefinitions/Source/Private/AssetDefinition_StaticMesh.cpp
@@ -22,6 +22,11 @@ UThumbnailInfo* UAssetDefinition_StaticMesh::LoadThumbnailInfo(const FAssetData&
 	return UE::Editor::FindOrCreateThumbnailInfo(InAssetData.GetAsset(), USceneThumbnailInfo::StaticClass());
 }
 
+FAssetOpenSupport UAssetDefinition_StaticMesh::GetAssetOpenSupport(const FAssetOpenSupportArgs& OpenSupportArgs) const
+{
+	return FAssetOpenSupport(OpenSupportArgs.OpenMethod,OpenSupportArgs.OpenMethod == EAssetOpenMethod::Edit || OpenSupportArgs.OpenMethod == EAssetOpenMethod::View); 
+}
+
 EAssetCommandResult UAssetDefinition_StaticMesh::OpenAssets(const FAssetOpenArgs& OpenArgs) const
 {
 	IStaticMeshEditorModule* StaticMeshEditorModule = &FModuleManager::LoadModuleChecked<IStaticMeshEditorModule>( "StaticMeshEditor" );
diff --git a/Engine/Plugins/Editor/EngineAssetDefinitions/Source/Private/AssetDefinition_StaticMesh.h b/Engine/Plugins/Editor/EngineAssetDefinitions/Source/Private/AssetDefinition_StaticMesh.h
--- a/Engine/Plugins/Editor/EngineAssetDefinitions/Source/Private/AssetDefinition_StaticMesh.h
+++ b/Engine/Plugins/Editor/EngineAssetDefinitions/Source/Private/AssetDefinition_StaticMesh.h
@@ -23,6 +23,8 @@ public:
 		return Categories;
 	}
 	virtual bool CanImport() const override { return true; }
+	virtual FAssetOpenSupport GetAssetOpenSupport(const FAssetOpenSupportArgs& OpenSupportArgs) const override;
+	
 	virtual UThumbnailInfo* LoadThumbnailInfo(const FAssetData& InAssetData) const override;
 	virtual EAssetCommandResult OpenAssets(const FAssetOpenArgs& OpenArgs) const override;
 	// UAssetDefinition End
diff --git a/Engine/Plugins/Editor/EngineAssetDefinitions/Source/Private/AssetDefinition_Texture.cpp b/Engine/Plugins/Editor/EngineAssetDefinitions/Source/Private/AssetDefinition_Texture.cpp
--- a/Engine/Plugins/Editor/EngineAssetDefinitions/Source/Private/AssetDefinition_Texture.cpp
+++ b/Engine/Plugins/Editor/EngineAssetDefinitions/Source/Private/AssetDefinition_Texture.cpp
@@ -31,6 +31,11 @@ EAssetCommandResult UAssetDefinition_Texture::OpenAssets(const FAssetOpenArgs& O
 	return EAssetCommandResult::Handled;
 }
 
+FAssetOpenSupport UAssetDefinition_Texture::GetAssetOpenSupport(const FAssetOpenSupportArgs& OpenSupportArgs) const
+{
+	return FAssetOpenSupport(OpenSupportArgs.OpenMethod,OpenSupportArgs.OpenMethod == EAssetOpenMethod::Edit || OpenSupportArgs.OpenMethod == EAssetOpenMethod::View); 
+}
+
 // Menu Extensions
 //--------------------------------------------------------------------
 
diff --git a/Engine/Plugins/Editor/EngineAssetDefinitions/Source/Private/AssetDefinition_Texture.h b/Engine/Plugins/Editor/EngineAssetDefinitions/Source/Private/AssetDefinition_Texture.h
--- a/Engine/Plugins/Editor/EngineAssetDefinitions/Source/Private/AssetDefinition_Texture.h
+++ b/Engine/Plugins/Editor/EngineAssetDefinitions/Source/Private/AssetDefinition_Texture.h
@@ -23,7 +23,8 @@ public:
 		return Categories;
 	}
 	virtual bool CanImport() const override { return true; }
-	
+
+	virtual FAssetOpenSupport GetAssetOpenSupport(const FAssetOpenSupportArgs& OpenSupportArgs) const override;
 	virtual EAssetCommandResult OpenAssets(const FAssetOpenArgs& OpenArgs) const override;
 	// UAssetDefinition End
 };
diff --git a/Engine/Source/Editor/ContentBrowser/Private/AssetViewWidgets.cpp b/Engine/Source/Editor/ContentBrowser/Private/AssetViewWidgets.cpp
--- a/Engine/Source/Editor/ContentBrowser/Private/AssetViewWidgets.cpp
+++ b/Engine/Source/Editor/ContentBrowser/Private/AssetViewWidgets.cpp
@@ -774,6 +774,7 @@ TSharedRef<SWidget> SAssetViewItem::CreateToolTipWidget() const
 			const FText ClassText = FText::Format(LOCTEXT("ClassName", "({0})"), GetAssetClassText());
 
 			FText PublicStateText;
+			const FSlateBrush* PublicStateIcon = nullptr;
 			FName PublicStateTextBorder = "ContentBrowser.TileViewTooltip.PillBorder";
 
 			// Create a box to hold every line of info in the body of the tooltip
@@ -828,7 +829,17 @@ TSharedRef<SWidget> SAssetViewItem::CreateToolTipWidget() const
 
 			if (!AssetItem->GetItem().CanEdit())
 			{
-				PublicStateText = LOCTEXT("ReadOnlyAssetState", "Read Only");
+				if(AssetItem->GetItem().CanView())
+				{
+					PublicStateText = LOCTEXT("ReadOnlyAssetState", "View / Read Only");
+					PublicStateIcon = FAppStyle::GetBrush("AssetEditor.ReadOnlyOpenable");
+
+				}
+				else
+				{
+					PublicStateText = LOCTEXT("ReadOnlyAssetState", "Read Only");
+					PublicStateIcon = FAppStyle::GetBrush("Icons.Lock");
+				}
 			}
 
 			if(!AssetItem->GetItem().IsSupported())
@@ -908,17 +919,38 @@ TSharedRef<SWidget> SAssetViewItem::CreateToolTipWidget() const
 							[
 								SNew(SBorder)
 								.BorderImage(FAppStyle::GetBrush(PublicStateTextBorder))
-								.Padding(FMargin(12.0f, 2.0f, 12.0f, 2.0f))
 								.Visibility(bIsPublicAssetUIEnabled && !PublicStateText.IsEmpty() ? EVisibility::Visible : EVisibility::Hidden)
+								.Padding(FMargin(12.0f, 2.0f, 12.0f, 2.0f))
 								[
-									SNew(STextBlock)
-									.Text(PublicStateText)
-									.HighlightText(HighlightText)
+									SNew(SHorizontalBox)
+									+SHorizontalBox::Slot()
+									.AutoWidth()
+									.HAlign(HAlign_Left)
+									.VAlign(VAlign_Center)
+									.Padding(0.0f)
+									[
+										SNew(SBox)
+										.Visibility(PublicStateIcon ? EVisibility::Visible : EVisibility::Collapsed)
+										.HeightOverride(16.0f)
+										.WidthOverride(16.0f)
+										[
+											SNew(SImage)
+											.Image(PublicStateIcon)
+										]
+											
+									]
+									+SHorizontalBox::Slot()
+									.HAlign(HAlign_Left)
+									.VAlign(VAlign_Center)
+									.Padding(4.0f, 0.0f, 0.0f, 0.0f)
+									[
+										SNew(STextBlock)
+										.Text(PublicStateText)
+										.HighlightText(HighlightText)
+									]
 								]
-								
 							]
 						]
-
 						+ SVerticalBox::Slot()
 						.AutoHeight()
 						[
diff --git a/Engine/Source/Editor/ContentBrowserData/Private/ContentBrowserItem.cpp b/Engine/Source/Editor/ContentBrowserData/Private/ContentBrowserItem.cpp
--- a/Engine/Source/Editor/ContentBrowserData/Private/ContentBrowserItem.cpp
+++ b/Engine/Source/Editor/ContentBrowserData/Private/ContentBrowserItem.cpp
@@ -448,6 +448,11 @@ bool FContentBrowserItem::CanPreview(FText* OutErrorMsg) const
 	return FContentBrowserItemHelper::CallDataSourceImpl<UContentBrowserDataSource>(*this, UE_PROJECTION_MEMBER(UContentBrowserDataSource, CanPreviewItem), OutErrorMsg);
 }
 
+bool FContentBrowserItem::CanView(FText* OutErrorMsg) const
+{
+	return FContentBrowserItemHelper::CallDataSourceImpl<UContentBrowserDataSource>(*this, UE_PROJECTION_MEMBER(UContentBrowserDataSource, CanViewItem), OutErrorMsg);
+}
+
 bool FContentBrowserItem::Preview() const
 {
 	return FContentBrowserItemHelper::CallDataSourceImpl<UContentBrowserDataSource>(*this, UE_PROJECTION_MEMBER(UContentBrowserDataSource, PreviewItem));
diff --git a/Engine/Source/Editor/ContentBrowserData/Public/ContentBrowserItem.h b/Engine/Source/Editor/ContentBrowserData/Public/ContentBrowserItem.h
--- a/Engine/Source/Editor/ContentBrowserData/Public/ContentBrowserItem.h
+++ b/Engine/Source/Editor/ContentBrowserData/Public/ContentBrowserItem.h
@@ -260,6 +260,15 @@ public:
 	 */
 	bool CanPreview(FText* OutErrorMsg = nullptr) const;
 
+	/**
+	 * Query whether the given item is can be viewed (a read-only asset editor), optionally providing error information if it cannot.
+	 *
+	 * @param OutErrorMessage Optional error message to fill on failure.
+	 *
+	 * @return True if the item can be viewed in a read-only editor, false otherwise.
+	 */
+	bool CanView(FText* OutErrorMsg = nullptr) const;
+
 	/**
 	 * Attempt to preview this item.
 	 * @return True if the item was previewed, false otherwise.
diff --git a/Engine/Source/Editor/EditorStyle/Private/StarshipStyle.cpp b/Engine/Source/Editor/EditorStyle/Private/StarshipStyle.cpp
--- a/Engine/Source/Editor/EditorStyle/Private/StarshipStyle.cpp
+++ b/Engine/Source/Editor/EditorStyle/Private/StarshipStyle.cpp
@@ -1245,6 +1245,10 @@ void FStarshipEditorStyle::FStyle::SetupGeneralStyles()
 
 		Set( "AssetEditor.ReimportAsset", new IMAGE_BRUSH( "Icons/icon_TextureEd_Reimport_40x", Icon40x40 ) );
 		Set( "AssetEditor.ReimportAsset.Small", new IMAGE_BRUSH( "Icons/icon_TextureEd_Reimport_40x", Icon20x20 ) );
+
+		Set( "AssetEditor.ReadOnlyBorder", new FSlateRoundedBoxBrush(FStyleColors::Foreground, 10.0f));
+		Set("AssetEditor.ReadOnlyOpenable", new IMAGE_BRUSH_SVG("Starship/AssetEditors/LockEye", Icon16x16));
+
 	}
 		
 	// Asset Thumbnail
@@ -2708,6 +2712,7 @@ void FStarshipEditorStyle::FStyle::SetupGeneralIcons()
 	Set("AssetEditor.CompileStatus.Overlay.Warning", new IMAGE_BRUSH_SVG("Starship/Blueprints/CompileStatus_Warning_Badge", Icon20x20, FStyleColors::Warning));
 	Set("AssetEditor.CompileStatus.Overlay.Good", new IMAGE_BRUSH_SVG("Starship/Blueprints/CompileStatus_Good_Badge", Icon20x20, FStyleColors::AccentGreen));
 	Set("AssetEditor.CompileStatus.Overlay.Error", new IMAGE_BRUSH_SVG("Starship/Blueprints/CompileStatus_Fail_Badge", Icon20x20, FStyleColors::Error));
+
 	
 	Set("Debug", new IMAGE_BRUSH_SVG( "Starship/Common/Debug", Icon16x16 ) );
 	Set("Modules", new IMAGE_BRUSH_SVG( "Starship/Common/Modules", Icon16x16 ) );
diff --git a/Engine/Source/Editor/PropertyEditor/Private/SDetailSingleItemRow.cpp b/Engine/Source/Editor/PropertyEditor/Private/SDetailSingleItemRow.cpp
--- a/Engine/Source/Editor/PropertyEditor/Private/SDetailSingleItemRow.cpp
+++ b/Engine/Source/Editor/PropertyEditor/Private/SDetailSingleItemRow.cpp
@@ -740,7 +740,8 @@ FReply SDetailSingleItemRow::OnMouseButtonUp(const FGeometry& MyGeometry, const
 			PulseAnimation.Play(SharedThis(this));
 			bIsHandled = true;
 		}
-		else if (PasteAction.CanExecute() && MouseEvent.GetEffectingButton() == EKeys::LeftMouseButton)
+		// Paste is disabled if property editing is disabled
+		else if (PasteAction.CanExecute() && MouseEvent.GetEffectingButton() == EKeys::LeftMouseButton && OwnerTreeNode.Pin()->GetDetailsView()->IsPropertyEditingEnabled())
 		{
 			PasteAction.Execute();
 			PulseAnimation.Play(SharedThis(this));
@@ -1045,7 +1046,12 @@ bool SDetailSingleItemRow::OnContextMenuOpening(FMenuBuilder& MenuBuilder)
 		PasteContentParams.InputBindingOverride = FInputChord(EModifierKey::Shift, EKeys::LeftMouseButton).GetInputText(bLongDisplayName);
 		PasteContentParams.IconOverride = FSlateIcon(FCoreStyle::Get().GetStyleSetName(), "GenericCommands.Paste");
 		PasteContentParams.DirectActions = PasteAction;
-		MenuBuilder.AddMenuEntry(PasteContentParams);
+
+		// Paste is disabled if property editing is disabled
+		if(OwnerTreeNode.Pin()->GetDetailsView()->IsPropertyEditingEnabled())
+		{
+			MenuBuilder.AddMenuEntry(PasteContentParams);
+		}
 	}
 
 	MenuBuilder.AddMenuEntry(
diff --git a/Engine/Source/Editor/PropertyEditor/Private/UserInterface/PropertyEditor/SPropertyEditorAsset.cpp b/Engine/Source/Editor/PropertyEditor/Private/UserInterface/PropertyEditor/SPropertyEditorAsset.cpp
--- a/Engine/Source/Editor/PropertyEditor/Private/UserInterface/PropertyEditor/SPropertyEditorAsset.cpp
+++ b/Engine/Source/Editor/PropertyEditor/Private/UserInterface/PropertyEditor/SPropertyEditorAsset.cpp
@@ -35,6 +35,7 @@
 #include "Presentation/PropertyEditor/PropertyEditor.h"
 #include "AssetThumbnail.h"
 #include "DetailWidgetRow.h"
+#include "Subsystems/AssetEditorSubsystem.h"
 
 #define LOCTEXT_NAMESPACE "PropertyEditor"
 
@@ -1121,7 +1122,20 @@ void SPropertyEditorAsset::OnOpenAssetEditor()
 			}
 		}
 
-		GEditor->EditObject( ObjectToEdit );
+		UAssetEditorSubsystem* AssetEditorSubsystem = GEditor->GetEditorSubsystem<UAssetEditorSubsystem>();
+
+		if(AssetEditorSubsystem)
+		{
+			FText ErrorMsg;
+			if(AssetEditorSubsystem->CanOpenEditorForAsset(ObjectToEdit, EAssetTypeActivationOpenedMethod::Edit, &ErrorMsg))
+			{
+				AssetEditorSubsystem->OpenEditorForAsset(ObjectToEdit); // Default opens in Edit Mode
+			}
+			else if(AssetEditorSubsystem->CanOpenEditorForAsset(ObjectToEdit, EAssetTypeActivationOpenedMethod::View, &ErrorMsg))
+			{
+				AssetEditorSubsystem->OpenEditorForAsset(ObjectToEdit, EToolkitMode::Standalone /* default */, TSharedPtr<IToolkitHost>() /* default */, true /* default */, EAssetTypeActivationOpenedMethod::View);
+			}
+		}
 	}
 }
 
diff --git a/Engine/Source/Editor/StaticMeshEditor/Private/StaticMeshEditor.cpp b/Engine/Source/Editor/StaticMeshEditor/Private/StaticMeshEditor.cpp
--- a/Engine/Source/Editor/StaticMeshEditor/Private/StaticMeshEditor.cpp
+++ b/Engine/Source/Editor/StaticMeshEditor/Private/StaticMeshEditor.cpp
@@ -176,32 +176,38 @@ void FStaticMeshEditor::RegisterTabSpawners(const TSharedRef<class FTabManager>&
 	InTabManager->RegisterTabSpawner( ViewportTabId, FOnSpawnTab::CreateSP(this, &FStaticMeshEditor::SpawnTab_Viewport) )
 		.SetDisplayName( LOCTEXT("ViewportTab", "Viewport") )
 		.SetGroup(WorkspaceMenuCategoryRef)
-		.SetIcon(FSlateIcon(FAppStyle::GetAppStyleSetName(), "LevelEditor.Tabs.Viewports"));
+		.SetIcon(FSlateIcon(FAppStyle::GetAppStyleSetName(), "LevelEditor.Tabs.Viewports"))
+		.SetReadOnlyBehavior(ETabReadOnlyBehavior::Custom);
 
 	InTabManager->RegisterTabSpawner( PropertiesTabId, FOnSpawnTab::CreateSP(this, &FStaticMeshEditor::SpawnTab_Properties) )
 		.SetDisplayName( LOCTEXT("PropertiesTab", "Details") )
 		.SetGroup(WorkspaceMenuCategoryRef)
-		.SetIcon(FSlateIcon(FAppStyle::GetAppStyleSetName(), "LevelEditor.Tabs.Details"));
+		.SetIcon(FSlateIcon(FAppStyle::GetAppStyleSetName(), "LevelEditor.Tabs.Details"))
+		.SetReadOnlyBehavior(ETabReadOnlyBehavior::Custom);
 
 	InTabManager->RegisterTabSpawner( SocketManagerTabId, FOnSpawnTab::CreateSP(this, &FStaticMeshEditor::SpawnTab_SocketManager) )
 		.SetDisplayName( LOCTEXT("SocketManagerTab", "Socket Manager") )
 		.SetGroup(WorkspaceMenuCategoryRef)
-		.SetIcon(FSlateIcon(FAppStyle::GetAppStyleSetName(), "StaticMeshEditor.Tabs.SocketManager"));
+		.SetIcon(FSlateIcon(FAppStyle::GetAppStyleSetName(), "StaticMeshEditor.Tabs.SocketManager"))
+		.SetReadOnlyBehavior(ETabReadOnlyBehavior::Custom);
 
 	InTabManager->RegisterTabSpawner( CollisionTabId, FOnSpawnTab::CreateSP(this, &FStaticMeshEditor::SpawnTab_Collision) )
 		.SetDisplayName( LOCTEXT("CollisionTab", "Convex Decomposition") )
 		.SetGroup(WorkspaceMenuCategoryRef)
-		.SetIcon(FSlateIcon(FAppStyle::GetAppStyleSetName(), "StaticMeshEditor.Tabs.ConvexDecomposition"));
+		.SetIcon(FSlateIcon(FAppStyle::GetAppStyleSetName(), "StaticMeshEditor.Tabs.ConvexDecomposition"))
+		.SetReadOnlyBehavior(ETabReadOnlyBehavior::Hidden);
 
 	InTabManager->RegisterTabSpawner( PreviewSceneSettingsTabId, FOnSpawnTab::CreateSP(this, &FStaticMeshEditor::SpawnTab_PreviewSceneSettings) )
 		.SetDisplayName( LOCTEXT("PreviewSceneTab", "Preview Scene Settings") )
 		.SetGroup(WorkspaceMenuCategoryRef)
-		.SetIcon(FSlateIcon(FAppStyle::GetAppStyleSetName(), "LevelEditor.Tabs.Details"));
+		.SetIcon(FSlateIcon(FAppStyle::GetAppStyleSetName(), "LevelEditor.Tabs.Details"))
+		.SetReadOnlyBehavior(ETabReadOnlyBehavior::Custom);
 
 	FTabSpawnerEntry& MenuEntry = InTabManager->RegisterTabSpawner( SecondaryToolbarTabId, FOnSpawnTab::CreateSP(this, &FStaticMeshEditor::SpawnTab_SecondaryToolbar) )
 		.SetDisplayName( LOCTEXT("ToolbarTab", "Secondary Toolbar") )
 		.SetGroup(WorkspaceMenuCategoryRef)
-		.SetIcon(FSlateIcon(FAppStyle::GetAppStyleSetName(), "Toolbar.Icon"));
+		.SetIcon(FSlateIcon(FAppStyle::GetAppStyleSetName(), "Toolbar.Icon"))
+		.SetReadOnlyBehavior(ETabReadOnlyBehavior::Hidden);
 
 	// Hide the menu item by default. It will be enabled only if the secondary toolbar is populated with extensions
 	SecondaryToolbarEntry = &MenuEntry;
@@ -359,6 +365,10 @@ void FStaticMeshEditor::InitStaticMeshEditor( const EToolkitMode::Type Mode, con
 	const bool bCreateDefaultToolbar = true;
 	FAssetEditorToolkit::InitAssetEditor( Mode, InitToolkitHost, StaticMeshEditorAppIdentifier, StandaloneDefaultLayout, bCreateDefaultToolbar, bCreateDefaultStandaloneMenu, ObjectToEdit );
 
+	StaticMeshDetailsView->SetIsPropertyEditingEnabledDelegate(FIsPropertyEditingEnabled::CreateLambda([this]
+	{
+		return GetOpenMethod() == EAssetOpenMethod::Edit;
+	}));
 	
 	TSharedPtr<class IToolkitHost> PinnedToolkitHost = ToolkitHost.Pin();
 	check(PinnedToolkitHost.IsValid());
@@ -493,12 +503,29 @@ void FStaticMeshEditor::ExtendMenu()
 
 		{
 			UToolMenu* Menu = UToolMenus::Get()->ExtendMenu("AssetEditor.StaticMeshEditor.MainMenu");
+
 			FToolMenuSection& Section = Menu->FindOrAddSection(NAME_None);
-			FToolMenuEntry& Entry = Section.AddSubMenu("Collision",
-				LOCTEXT("StaticMeshEditorCollisionMenu", "Collision"),
-				LOCTEXT("StaticMeshEditorCollisionMenu_ToolTip", "Opens a menu with commands for editing this mesh's collision"),
-				FNewToolMenuChoice());
-			Entry.InsertPosition = FToolMenuInsert("Asset", EToolMenuInsertType::After);
+
+			Section.AddDynamicEntry("CollisionDynamic", FNewToolMenuSectionDelegate::CreateLambda([](FToolMenuSection& InSection)
+			{
+				if (TSharedPtr<FStaticMeshEditor> StaticMeshEditor = StaticMeshEditor::GetStaticMeshEditorFromMenuContext(InSection))
+				{
+					// These entries modify the mesh, so we don't want them in read only mode
+					if(StaticMeshEditor->GetOpenMethod() != EAssetOpenMethod::Edit)
+					{
+						return;
+					}
+					
+					FToolMenuEntry& Entry = InSection.AddSubMenu("Collision",
+						LOCTEXT("StaticMeshEditorCollisionMenu", "Collision"),
+						LOCTEXT("StaticMeshEditorCollisionMenu_ToolTip", "Opens a menu with commands for editing this mesh's collision"),
+						FNewToolMenuChoice());
+						
+					Entry.InsertPosition = FToolMenuInsert("Asset", EToolMenuInsertType::After);
+				}
+			}));
+			
+			
 		}
 	}
 
@@ -865,45 +892,74 @@ void FStaticMeshEditor::ExtendToolBar()
 	{
 		UToolMenu* Menu = UToolMenus::Get()->ExtendMenu("AssetEditor.StaticMeshEditor.ToolBar");
 
-		auto ConstructReimportContextMenu = [](UToolMenu* InMenu)
-		{
-			FToolMenuSection& Section = InMenu->AddSection("Reimport");
-			Section.AddMenuEntry(FStaticMeshEditorCommands::Get().ReimportMesh);
-			Section.AddMenuEntry(FStaticMeshEditorCommands::Get().ReimportAllMesh);
-		};
-
 		{
 			FToolMenuSection& Section = Menu->AddSection("Mesh");
 			Section.InsertPosition = FToolMenuInsert("Asset", EToolMenuInsertType::After);
 
-			FToolMenuEntry& ReimportMeshEntry = Section.AddEntry(FToolMenuEntry::InitToolBarButton(FStaticMeshEditorCommands::Get().ReimportMesh));
-			ReimportMeshEntry.StyleNameOverride = "CalloutToolbar";
-			
-			FToolMenuEntry& ReimportContextMenuEntry = Section.AddEntry(FToolMenuEntry::InitComboButton(
-				"ReimportContextMenu",
-				FUIAction(),
-				FNewToolMenuDelegate::CreateLambda(ConstructReimportContextMenu),
-				TAttribute<FText>(),
-				TAttribute<FText>(),
-				TAttribute<FSlateIcon>(),
-				true
-			));
-			ReimportContextMenuEntry.StyleNameOverride = "CalloutToolbar";
+			Section.AddDynamicEntry("MeshDynamic", FNewToolMenuSectionDelegate::CreateLambda([](FToolMenuSection& InSection)
+			{
+				if (TSharedPtr<FStaticMeshEditor> StaticMeshEditor = StaticMeshEditor::GetStaticMeshEditorFromMenuContext(InSection))
+				{
+					// These entries modify the mesh, so we don't want them in read only mode
+					if(StaticMeshEditor->GetOpenMethod() != EAssetOpenMethod::Edit)
+					{
+						return;
+					}
+					
+					auto ConstructReimportContextMenu = [](UToolMenu* InMenu)
+					{
+						FToolMenuSection& Section = InMenu->AddSection("Reimport");
+						Section.AddMenuEntry(FStaticMeshEditorCommands::Get().ReimportMesh);
+						Section.AddMenuEntry(FStaticMeshEditorCommands::Get().ReimportAllMesh);
+					};
+					
+					FToolMenuEntry& ReimportMeshEntry = InSection.AddEntry(FToolMenuEntry::InitToolBarButton(FStaticMeshEditorCommands::Get().ReimportMesh));
+					ReimportMeshEntry.StyleNameOverride = "CalloutToolbar";
+								
+					FToolMenuEntry& ReimportContextMenuEntry = InSection.AddEntry(FToolMenuEntry::InitComboButton(
+						"ReimportContextMenu",
+						FUIAction(),
+						FNewToolMenuDelegate::CreateLambda(ConstructReimportContextMenu),
+						TAttribute<FText>(),
+						TAttribute<FText>(),
+						TAttribute<FSlateIcon>(),
+						true
+					));
+					ReimportContextMenuEntry.StyleNameOverride = "CalloutToolbar";
+
+				}
+			}));
+
 		}
 
 		{
 			FToolMenuSection& Section = Menu->AddSection("Command");
 			Section.InsertPosition = FToolMenuInsert("Asset", EToolMenuInsertType::After);
 
-			FToolMenuEntry& CollisionEntry = Section.AddEntry(FToolMenuEntry::InitComboButton(
-				"Collision",
-				FUIAction(),
-				FNewToolMenuChoice(), // let registered menu be looked up by name "AssetEditor.StaticMeshEditor.ToolBar.Collision"
-				LOCTEXT("Collision_Label", "Collision"),
-				LOCTEXT("Collision_Tooltip", "Collision drawing options"),
-				FSlateIcon(FAppStyle::GetAppStyleSetName(), "StaticMeshEditor.SetShowCollision")
-			));
-			CollisionEntry.StyleNameOverride = "CalloutToolbar";
+			Section.AddDynamicEntry("MeshDynamic", FNewToolMenuSectionDelegate::CreateLambda([](FToolMenuSection& InSection)
+			{
+				if (TSharedPtr<FStaticMeshEditor> StaticMeshEditor = StaticMeshEditor::GetStaticMeshEditorFromMenuContext(InSection))
+				{
+					// These entries modify the mesh, so we don't want them in read only mode
+					if(StaticMeshEditor->GetOpenMethod() != EAssetOpenMethod::Edit)
+					{
+						return;
+					}
+					
+					FToolMenuEntry& CollisionEntry = InSection.AddEntry(FToolMenuEntry::InitComboButton(
+					"Collision",
+					FUIAction(),
+					FNewToolMenuChoice(), // let registered menu be looked up by name "AssetEditor.StaticMeshEditor.ToolBar.Collision"
+					LOCTEXT("Collision_Label", "Collision"),
+					LOCTEXT("Collision_Tooltip", "Collision drawing options"),
+					FSlateIcon(FAppStyle::GetAppStyleSetName(), "StaticMeshEditor.SetShowCollision")
+					));
+					
+					CollisionEntry.StyleNameOverride = "CalloutToolbar";
+				}
+			}));
+
+			
 
 			Section.AddDynamicEntry("UVToolbarDynamic", FNewToolMenuSectionDelegate::CreateLambda([](FToolMenuSection& InSection)
 			{
@@ -923,21 +979,24 @@ void FStaticMeshEditor::ExtendToolBar()
 		}
 	}
 
-	{
-	IStaticMeshEditorModule* StaticMeshEditorModule = &FModuleManager::LoadModuleChecked<IStaticMeshEditorModule>("StaticMeshEditor");
 
-	TArray<IStaticMeshEditorModule::FStaticMeshEditorToolbarExtender> ToolbarExtenderDelegates = StaticMeshEditorModule->GetAllStaticMeshEditorToolbarExtenders();
-	for (auto& ToolbarExtenderDelegate : ToolbarExtenderDelegates)
+	// Extensions are currently disabled in read-only mode, if they are desired to be added in the future we should move this code to the individual extensions
+	if(GetOpenMethod() == EAssetOpenMethod::Edit)
 	{
-		if (ToolbarExtenderDelegate.IsBound())
+		IStaticMeshEditorModule* StaticMeshEditorModule = &FModuleManager::LoadModuleChecked<IStaticMeshEditorModule>("StaticMeshEditor");
+
+		TArray<IStaticMeshEditorModule::FStaticMeshEditorToolbarExtender> ToolbarExtenderDelegates = StaticMeshEditorModule->GetAllStaticMeshEditorToolbarExtenders();
+		for (auto& ToolbarExtenderDelegate : ToolbarExtenderDelegates)
 		{
-			AddToolbarExtender(ToolbarExtenderDelegate.Execute(GetToolkitCommands(), SharedThis(this)));
+			if (ToolbarExtenderDelegate.IsBound())
+			{
+				AddToolbarExtender(ToolbarExtenderDelegate.Execute(GetToolkitCommands(), SharedThis(this)));
+			}
 		}
-	}
 
-	EditorToolbarExtender = StaticMeshEditorModule->GetToolBarExtensibilityManager()->GetAllExtenders(GetToolkitCommands(), GetEditingObjects());
-	AddToolbarExtender(EditorToolbarExtender);
-	AddSecondaryToolbarExtender(StaticMeshEditorModule->GetSecondaryToolBarExtensibilityManager()->GetAllExtenders(GetToolkitCommands(), GetEditingObjects()));
+		EditorToolbarExtender = StaticMeshEditorModule->GetToolBarExtensibilityManager()->GetAllExtenders(GetToolkitCommands(), GetEditingObjects());
+		AddToolbarExtender(EditorToolbarExtender);
+		AddSecondaryToolbarExtender(StaticMeshEditorModule->GetSecondaryToolBarExtensibilityManager()->GetAllExtenders(GetToolkitCommands(), GetEditingObjects()));
 	}
 
 }
@@ -1607,21 +1666,27 @@ void FStaticMeshEditor::GenerateUVChannelComboList(UToolMenu* InMenu)
 		}
 	}
 
-	// Add UV editing functions
+
+	if (TSharedPtr<FStaticMeshEditor> StaticMeshEditor = StaticMeshEditor::GetStaticMeshEditorFromMenuContext(InMenu))
 	{
-		FToolMenuSection& Section = InMenu->AddSection("UVActionOptions");
+		// These actions modify the mesh, so we don't want them in read only mode
+		if(StaticMeshEditor->GetOpenMethod() == EAssetOpenMethod::Edit)
+		{
+			FToolMenuSection& Section = InMenu->AddSection("UVActionOptions");
 
-		FUIAction MenuAction;
-		MenuAction.ExecuteAction.BindSP(this, &FStaticMeshEditor::RemoveCurrentUVChannel);
-		MenuAction.CanExecuteAction.BindSP(this, &FStaticMeshEditor::CanRemoveUVChannel);
-		Section.AddMenuEntry(
-			"Remove_UVChannel",
-			LOCTEXT("Remove_UVChannel", "Remove Selected"),
-			LOCTEXT("Remove_UVChannel_ToolTip", "Remove currently selected UV channel from the static mesh"),
-			FSlateIcon(),
-			MenuAction,
-			EUserInterfaceActionType::Button
-		);
+			FUIAction MenuAction;
+			MenuAction.ExecuteAction.BindSP(this, &FStaticMeshEditor::RemoveCurrentUVChannel);
+			MenuAction.CanExecuteAction.BindSP(this, &FStaticMeshEditor::CanRemoveUVChannel);
+			Section.AddMenuEntry(
+				"Remove_UVChannel",
+				LOCTEXT("Remove_UVChannel", "Remove Selected"),
+				LOCTEXT("Remove_UVChannel_ToolTip", "Remove currently selected UV channel from the static mesh"),
+				FSlateIcon(),
+				MenuAction,
+				EUserInterfaceActionType::Button
+			);
+		}
+		
 	}
 }
 
@@ -2454,7 +2519,8 @@ void FStaticMeshEditor::OnConvexDecomposition()
 bool FStaticMeshEditor::OnRequestClose(EAssetEditorCloseReason InCloseReason)
 {
 	bool bAllowClose = true;
-	if (InCloseReason != EAssetEditorCloseReason::AssetForceDeleted && StaticMeshDetails.IsValid() && StaticMeshDetails.Pin()->IsApplyNeeded())
+	// If we are in read only mode, don't show the save prompt
+	if (GetOpenMethod() != EAssetOpenMethod::View && InCloseReason != EAssetEditorCloseReason::AssetForceDeleted && StaticMeshDetails.IsValid() && StaticMeshDetails.Pin()->IsApplyNeeded())
 	{
 		// find out the user wants to do with this dirty material
 		EAppReturnType::Type YesNoCancelReply = FMessageDialog::Open(
diff --git a/Engine/Source/Editor/TextureEditor/Private/TextureEditorToolkit.cpp b/Engine/Source/Editor/TextureEditor/Private/TextureEditorToolkit.cpp
--- a/Engine/Source/Editor/TextureEditor/Private/TextureEditorToolkit.cpp
+++ b/Engine/Source/Editor/TextureEditor/Private/TextureEditorToolkit.cpp
@@ -174,12 +174,14 @@ void FTextureEditorToolkit::RegisterTabSpawners( const TSharedRef<class FTabMana
 	InTabManager->RegisterTabSpawner(ViewportTabId, FOnSpawnTab::CreateSP(this, &FTextureEditorToolkit::HandleTabSpawnerSpawnViewport))
 		.SetDisplayName(LOCTEXT("ViewportTab", "Viewport"))
 		.SetGroup(WorkspaceMenuCategoryRef)
-		.SetIcon(FSlateIcon(FAppStyle::GetAppStyleSetName(), "LevelEditor.Tabs.Viewports"));
+		.SetIcon(FSlateIcon(FAppStyle::GetAppStyleSetName(), "LevelEditor.Tabs.Viewports"))
+		.SetReadOnlyBehavior(ETabReadOnlyBehavior::Custom);
 
 	InTabManager->RegisterTabSpawner(PropertiesTabId, FOnSpawnTab::CreateSP(this, &FTextureEditorToolkit::HandleTabSpawnerSpawnProperties))
 		.SetDisplayName(LOCTEXT("PropertiesTab", "Details") )
 		.SetGroup(WorkspaceMenuCategoryRef)
-		.SetIcon(FSlateIcon(FAppStyle::GetAppStyleSetName(), "LevelEditor.Tabs.Details"));
+		.SetIcon(FSlateIcon(FAppStyle::GetAppStyleSetName(), "LevelEditor.Tabs.Details"))
+		.SetReadOnlyBehavior(ETabReadOnlyBehavior::Custom);
 
 	InTabManager->RegisterTabSpawner(OodleTabId, FOnSpawnTab::CreateSP(this, &FTextureEditorToolkit::HandleTabSpawnerSpawnOodle))
 		.SetDisplayName(LOCTEXT("OodleTab", "Oodle"))
@@ -316,6 +318,11 @@ void FTextureEditorToolkit::InitTextureEditor( const EToolkitMode::Type Mode, co
 	ITextureEditorModule* TextureEditorModule = &FModuleManager::LoadModuleChecked<ITextureEditorModule>("TextureEditor");
 	AddMenuExtender(TextureEditorModule->GetMenuExtensibilityManager()->GetAllExtenders(GetToolkitCommands(), GetEditingObjects()));
 
+	TexturePropertiesWidget->SetIsPropertyEditingEnabledDelegate(FIsPropertyEditingEnabled::CreateLambda([this]
+	{
+		return GetOpenMethod() == EAssetOpenMethod::Edit;
+	}));
+	
 	ExtendToolBar();
 
 	RegenerateMenusAndToolbars();
@@ -1828,8 +1835,12 @@ void FTextureEditorToolkit::ExtendToolBar()
 
 	AddToolbarExtender(ToolbarExtender);
 
-	ITextureEditorModule* TextureEditorModule = &FModuleManager::LoadModuleChecked<ITextureEditorModule>("TextureEditor");
-	AddToolbarExtender(TextureEditorModule->GetToolBarExtensibilityManager()->GetAllExtenders(GetToolkitCommands(), GetEditingObjects()));
+	// Extensions are currently disabled in read-only mode, if they are desired to be added in the future we should move this code to the individual extensions
+	if(GetOpenMethod() == EAssetOpenMethod::Edit)
+	{
+		ITextureEditorModule* TextureEditorModule = &FModuleManager::LoadModuleChecked<ITextureEditorModule>("TextureEditor");
+		AddToolbarExtender(TextureEditorModule->GetToolBarExtensibilityManager()->GetAllExtenders(GetToolkitCommands(), GetEditingObjects()));
+	}
 }
 
 void FTextureEditorToolkit::FillToolbar(FToolBarBuilder& ToolbarBuilder)
@@ -1847,12 +1858,17 @@ void FTextureEditorToolkit::FillToolbar(FToolBarBuilder& ToolbarBuilder)
 	UCurveLinearColorAtlas* Atlas = Cast<UCurveLinearColorAtlas>(GetTexture());
 	if (!Atlas)
 	{
-		ToolbarBuilder.BeginSection("TextureMisc");
+		// These actions don't make sense in read-only mode
+		if(GetOpenMethod() == EAssetOpenMethod::Edit)
 		{
-			ToolbarBuilder.AddToolBarButton(FTextureEditorCommands::Get().CompressNow);
-			ToolbarBuilder.AddToolBarButton(FTextureEditorCommands::Get().Reimport);
+			ToolbarBuilder.BeginSection("TextureMisc");
+			{
+				ToolbarBuilder.AddToolBarButton(FTextureEditorCommands::Get().CompressNow);
+				ToolbarBuilder.AddToolBarButton(FTextureEditorCommands::Get().Reimport);
+			}
+			ToolbarBuilder.EndSection();
 		}
-		ToolbarBuilder.EndSection();
+		
 
 		ToolbarBuilder.BeginSection("Channels");
 		{
diff --git a/Engine/Source/Editor/UnrealEd/Private/SSocketManager.cpp b/Engine/Source/Editor/UnrealEd/Private/SSocketManager.cpp
--- a/Engine/Source/Editor/UnrealEd/Private/SSocketManager.cpp
+++ b/Engine/Source/Editor/UnrealEd/Private/SSocketManager.cpp
@@ -50,14 +50,18 @@ class SSocketDisplayItem : public SMultiColumnTableRow< TSharedPtr<SocketListIte
 
 public:
 	
-	SLATE_BEGIN_ARGS( SSocketDisplayItem )
-		{}
+	SLATE_BEGIN_ARGS( SSocketDisplayItem ):
+		_ReadOnly(false)
+		{} 
 
 		/** The socket this item displays. */
 		SLATE_ARGUMENT( TWeakPtr< SocketListItem >, SocketItem )
 
 		/** Pointer back to the socket manager */
 		SLATE_ARGUMENT( TWeakPtr< SSocketManager >, SocketManagerPtr )
+
+		/** Whether the widget should be editable or not */
+		SLATE_ARGUMENT( bool, ReadOnly)
 	SLATE_END_ARGS()
 
 	/**
@@ -69,6 +73,7 @@ public:
 	{
 		SocketItem = InArgs._SocketItem;
 		SocketManagerPtr = InArgs._SocketManagerPtr;
+		bReadOnly = InArgs._ReadOnly;
 
 		ImportedSocketBrush = FAppStyle::Get().GetBrush("Icons.Import");
 		NotImportedSocketBrush = FAppStyle::Get().GetBrush("NoBrush");
@@ -87,6 +92,7 @@ public:
 			TSharedPtr< SInlineEditableTextBlock > InlineWidget;
 
 			SAssignNew( InlineWidget, SInlineEditableTextBlock )
+			.IsReadOnly(bReadOnly)
 			.Text( this, &SSocketDisplayItem::GetSocketName )
 			.OnVerifyTextChanged( this, &SSocketDisplayItem::OnVerifySocketNameChanged )
 			.OnTextCommitted( this, &SSocketDisplayItem::OnCommitSocketName )
@@ -211,6 +217,8 @@ private:
 
 	const FSlateBrush* ImportedSocketBrush;
 	const FSlateBrush* NotImportedSocketBrush;
+
+	bool bReadOnly;
 };
 
 TSharedPtr<ISocketManager> ISocketManager::CreateSocketManager(TSharedPtr<class IStaticMeshEditor> InStaticMeshEditor, FSimpleDelegate InOnSocketSelectionChanged )
@@ -218,7 +226,8 @@ TSharedPtr<ISocketManager> ISocketManager::CreateSocketManager(TSharedPtr<class
 	TSharedPtr<SSocketManager> SocketManager;
 	SAssignNew(SocketManager, SSocketManager)
 		.StaticMeshEditorPtr(InStaticMeshEditor)
-		.OnSocketSelectionChanged( InOnSocketSelectionChanged );
+		.OnSocketSelectionChanged( InOnSocketSelectionChanged )
+		.ReadOnly(InStaticMeshEditor->GetOpenMethod() == EAssetOpenMethod::View);
 
 	TSharedPtr<ISocketManager> ISocket;
 	ISocket = StaticCastSharedPtr<ISocketManager>(SocketManager);
@@ -231,6 +240,8 @@ void SSocketManager::Construct(const FArguments& InArgs)
 
 	OnSocketSelectionChanged = InArgs._OnSocketSelectionChanged;
 
+	bReadOnly = InArgs._ReadOnly;
+
 	TSharedPtr<IStaticMeshEditor> StaticMeshEditorPinned = StaticMeshEditorPtr.Pin();
 	if (!StaticMeshEditorPinned.IsValid())
 	{
@@ -252,6 +263,16 @@ void SSocketManager::Construct(const FArguments& InArgs)
 
 	FPropertyEditorModule& PropertyModule = FModuleManager::LoadModuleChecked<FPropertyEditorModule>("PropertyEditor");
 	SocketDetailsView = PropertyModule.CreateDetailView(Args);
+	
+	SocketDetailsView->SetIsPropertyEditingEnabledDelegate(FIsPropertyEditingEnabled::CreateLambda([this]
+	{
+		if(TSharedPtr<IStaticMeshEditor> StaticMeshEditorPinned = StaticMeshEditorPtr.Pin())
+		{
+			return StaticMeshEditorPinned->GetOpenMethod() == EAssetOpenMethod::Edit;
+		}
+
+		return true;
+	}));
 
 	WorldSpaceRotation = FVector::ZeroVector;
 
@@ -299,6 +320,7 @@ void SSocketManager::Construct(const FArguments& InArgs)
 							.ButtonStyle(FAppStyle::Get(), "SimpleButton")
 							.ToolTipText(LOCTEXT("CreateSocket", "Create Socket"))
 							.OnClicked(this, &SSocketManager::CreateSocket_Execute)
+							.Visibility(this, &SSocketManager::CreateSocket_IsVisible )
 							.VAlign(VAlign_Center)
 							[
 								SNew(SImage)
@@ -449,7 +471,8 @@ TSharedRef< ITableRow > SSocketManager::MakeWidgetFromOption( TSharedPtr<SocketL
 {
 	return SNew( SSocketDisplayItem, OwnerTable )
 				.SocketItem(InItem)
-				.SocketManagerPtr(SharedThis(this));
+				.SocketManagerPtr(SharedThis(this))
+				.ReadOnly(bReadOnly);
 }
 
 void SSocketManager::CreateSocket()
@@ -661,6 +684,11 @@ FReply SSocketManager::CreateSocket_Execute()
 	return FReply::Handled();
 }
 
+EVisibility SSocketManager::CreateSocket_IsVisible() const
+{
+	return bReadOnly ? EVisibility::Collapsed : EVisibility::Visible;
+}	
+
 FText SSocketManager::GetSocketHeaderText() const
 {
 	UStaticMesh* CurrentStaticMesh = nullptr;
diff --git a/Engine/Source/Editor/UnrealEd/Private/SSocketManager.h b/Engine/Source/Editor/UnrealEd/Private/SSocketManager.h
--- a/Engine/Source/Editor/UnrealEd/Private/SSocketManager.h
+++ b/Engine/Source/Editor/UnrealEd/Private/SSocketManager.h
@@ -23,8 +23,11 @@ class SSocketManager : public ISocketManager, public FNotifyHook
 {
 public:
 	SLATE_BEGIN_ARGS( SSocketManager ) :
-		 _StaticMeshEditorPtr(){}
+		 _StaticMeshEditorPtr(),
+		_ReadOnly(false)
+		{}
 		SLATE_ARGUMENT( TSharedPtr< IStaticMeshEditor >, StaticMeshEditorPtr )
+		SLATE_ARGUMENT(bool, ReadOnly)
 		SLATE_EVENT( FSimpleDelegate, OnSocketSelectionChanged )
 	SLATE_END_ARGS()
 
@@ -76,6 +79,9 @@ private:
 	/** Callback for the Create Socket button. */
 	FReply CreateSocket_Execute();
 
+	/** Determines if the Create Socket button is visible. */
+	EVisibility CreateSocket_IsVisible() const;
+
 	FText GetSocketHeaderText() const;
 
 	/** Callback for when the socket name textbox is changed, verifies the name is not a duplicate. */
@@ -130,4 +136,7 @@ private:
 
 	/** Points to an item that is being requested to be renamed */
 	TWeakPtr<SocketListItem> DeferredRenameRequest;
+
+	/** true if the owning asset editor is in read only mode*/
+	bool bReadOnly;
 };
diff --git a/Engine/Source/Editor/UnrealEd/Private/Subsystems/AssetEditorSubsystem.cpp b/Engine/Source/Editor/UnrealEd/Private/Subsystems/AssetEditorSubsystem.cpp
--- a/Engine/Source/Editor/UnrealEd/Private/Subsystems/AssetEditorSubsystem.cpp
+++ b/Engine/Source/Editor/UnrealEd/Private/Subsystems/AssetEditorSubsystem.cpp
@@ -486,9 +486,20 @@ bool UAssetEditorSubsystem::OpenEditorForAsset(UObject* Asset, const EToolkitMod
 {
 	SCOPE_STALL_REPORTER(UAssetEditorSubsystem::OpenEditorForAsset, 2.0);
 
-	if (!Asset)
+	FText ErrorMessage;
+	if(!CanOpenEditorForAsset(Asset, OpenedMethod, &ErrorMessage))
 	{
-		UE_LOG(LogAssetEditorSubsystem, Error, TEXT("Opening Asset editor failed because asset is null"));
+		// We also log the error if the asset was null
+		if(!Asset)
+		{
+			UE_LOG(LogAssetEditorSubsystem, Error, TEXT("%s"), *ErrorMessage.ToString());
+		}
+		
+		if (TSharedPtr<SNotificationItem> InfoItem = FSlateNotificationManager::Get().AddNotification(FNotificationInfo(ErrorMessage)))
+        {
+        	InfoItem->SetCompletionState(SNotificationItem::CS_Fail);
+        }
+
 		return false;
 	}
 
@@ -500,28 +511,6 @@ bool UAssetEditorSubsystem::OpenEditorForAsset(UObject* Asset, const EToolkitMod
 
 	const bool bBringToFrontIfOpen = true;
 
-	if (OpenedMethod == EAssetTypeActivationOpenedMethod::Edit)
-	{
-		if (UPackage* Package = Asset->GetOutermost())
-		{
-			// Don't open asset editors for cooked packages
-			if (Package->bIsCookedForEditor)
-			{
-				if (TSharedPtr<SNotificationItem> InfoItem = FSlateNotificationManager::Get().AddNotification(FNotificationInfo(LOCTEXT("NotifyBlockedByCookedAsset", "Unable to Edit Cooked asset"))))
-				{
-					InfoItem->SetCompletionState(SNotificationItem::CS_Fail);
-				}
-				return false;
-			}
-
-			if (!AssetToolsModule.Get().GetWritableFolderPermissionList()->PassesStartsWithFilter(Package->GetName()))
-			{
-				AssetToolsModule.Get().NotifyBlockedByWritableFolderFilter();
-				return false;
-			}
-		}
-	}
-
 	AssetEditorRequestOpenEvent.Broadcast(Asset);
 
 	if (FindEditorForAsset(Asset, bBringToFrontIfOpen) != nullptr)
@@ -539,6 +528,12 @@ bool UAssetEditorSubsystem::OpenEditorForAsset(UObject* Asset, const EToolkitMod
 
 	UE_LOG(LogAssetEditorSubsystem, Log, TEXT("Opening Asset editor for %s"), *Asset->GetFullName());
 
+	const FString AssetPath = Asset->GetPathName();
+	EAssetOpenMethod AssetOpenMethod = OpenedMethod == EAssetTypeActivationOpenedMethod::View ? EAssetOpenMethod::View : EAssetOpenMethod::Edit;
+
+	// We store the open method for this asset, so that FAssetEditorToolkit can query us for this information during init
+	AssetOpenMethodCache.Add(AssetPath, AssetOpenMethod);
+
 	TWeakPtr<IAssetTypeActions> AssetTypeActions = AssetToolsModule.Get().GetAssetTypeActionsForClass(Asset->GetClass());
 
 	EToolkitMode::Type ActualToolkitMode = ToolkitMode;
@@ -564,17 +559,13 @@ bool UAssetEditorSubsystem::OpenEditorForAsset(UObject* Asset, const EToolkitMod
 		FAssetEditorToolkit::SetPreviousWorldCentricToolkitHostForNewAssetEditor(OpenedFromLevelEditor.ToSharedRef());
 	}
 
-	// Disallow opening an asset editor for classes
-	bool bCanSummonSimpleAssetEditor = !Asset->IsA<UClass>();
-
-	if (AssetTypeActions.IsValid() && AssetTypeActions.Pin()->SupportsOpenedMethod(OpenedMethod))
+	if (AssetTypeActions.IsValid())
 	{
 		TArray<UObject*> AssetsToEdit;
 		AssetsToEdit.Add(Asset);
 
 		// Some assets (like UWorlds) may be destroyed and recreated as part of opening. To protect against this, keep the path to the asset and try to re-find it if it disappeared.
 		TWeakObjectPtr<UObject> WeakAsset = Asset;
-		const FString AssetPath = Asset->GetPathName();
 
 		AssetTypeActions.Pin()->OpenAssetEditor(AssetsToEdit, OpenedMethod, ActualToolkitMode == EToolkitMode::WorldCentric ? OpenedFromLevelEditor : TSharedPtr<IToolkitHost>());
 
@@ -586,7 +577,7 @@ bool UAssetEditorSubsystem::OpenEditorForAsset(UObject* Asset, const EToolkitMod
 
 		AssetEditorOpenedEvent.Broadcast(Asset);
 	}
-	else if (bCanSummonSimpleAssetEditor)
+	else
 	{
 		// No asset type actions for this asset. Just use a properties editor.
 		FSimpleAssetEditor::CreateEditor(ActualToolkitMode, ActualToolkitMode == EToolkitMode::WorldCentric ? OpenedFromLevelEditor : TSharedPtr<IToolkitHost>(), Asset);
@@ -602,14 +593,17 @@ bool UAssetEditorSubsystem::OpenEditorForAsset(UObject* Asset, const EToolkitMod
 	{
 		if (Asset->IsAsset() && !Asset->IsA(UMapBuildDataRegistry::StaticClass()))
 		{
-			FString AssetPath = Asset->GetOuter()->GetPathName();
-			if (FPackageName::IsValidLongPackageName(AssetPath))
+			FString AssetOuterPath = Asset->GetOuter()->GetPathName();
+			if (FPackageName::IsValidLongPackageName(AssetOuterPath))
 			{
-				RecentAssetsList->AddMRUItem(AssetPath);
+				RecentAssetsList->AddMRUItem(AssetOuterPath);
 				CullRecentAssetEditorsMap();
 			}
 		}
 	}
+
+	// Since the Asset Editor has finished init once we are here, we can remove the open method from the cache
+	AssetOpenMethodCache.Remove(AssetPath);
 	
 	return true;
 }
@@ -649,10 +643,11 @@ bool UAssetEditorSubsystem::OpenEditorForAssets_Advanced(const TArray <UObject*>
 		TArray<UObject*> SkipOpenAssets;
 		for (UObject* Asset : Assets)
 		{
-			// If any of the assets are already open or the package is cooked,
+			// If any of the assets are already open or they cannot be opened in this open method
 			// remove them from the list of assets to open an editor for
 			UPackage* Package = Asset->GetOutermost();
-			if (FindEditorForAsset(Asset, true) != nullptr || (OpenedMethod == EAssetTypeActivationOpenedMethod::Edit && Package && Package->bIsCookedForEditor))
+			FText ErrorMessage;
+			if (FindEditorForAsset(Asset, true) != nullptr || CanOpenEditorForAsset(Asset, OpenedMethod, &ErrorMessage))
 			{
 				SkipOpenAssets.Add(Asset);
 			}
@@ -711,11 +706,16 @@ bool UAssetEditorSubsystem::OpenEditorForAssets_Advanced(const TArray <UObject*>
 						: WeakAsset(InWeakAsset), AssetPath(InAssetPath) {}
 				};
 
+				EAssetOpenMethod AssetOpenMethod = OpenedMethod == EAssetTypeActivationOpenedMethod::View ? EAssetOpenMethod::View : EAssetOpenMethod::Edit;
+
 				TArray<FLocalAssetInfo> AssetInfoList;
 				AssetInfoList.Reserve(Assets.Num());
 				for (UObject* Asset : Assets)
 				{
 					AssetInfoList.Add(FLocalAssetInfo(Asset, Asset->GetPathName()));
+
+					// We store the open method for this asset, so that FAssetEditorToolkit can query us for this information during init
+					AssetOpenMethodCache.Add(Asset->GetPathName(), AssetOpenMethod);
 				}
 
 				// How to handle multiple assets is left up to the type actions (i.e. open a single shared editor or an editor for each)
@@ -731,6 +731,9 @@ bool UAssetEditorSubsystem::OpenEditorForAssets_Advanced(const TArray <UObject*>
 					{
 						Asset = FindObject<UObject>(nullptr, *AssetInfo.AssetPath);
 					}
+
+					// Since the Asset Editor has finished init once we are here, we can remove the open method from the cache
+					AssetOpenMethodCache.Remove(AssetInfo.AssetPath);
 				}
 
 				//@todo if needed, broadcast the event for every asset. It is possible, however, that a single shared editor was opened by the AssetTypeActions, not an editor for each asset.
@@ -760,19 +763,72 @@ bool UAssetEditorSubsystem::OpenEditorForAssets(const TArray<UObject*>& Assets,
 	return OpenEditorForAssets_Advanced(Assets, EToolkitMode::Standalone, TSharedPtr<IToolkitHost>(), OpenedMethod);
 }
 
+TOptional<EAssetOpenMethod> UAssetEditorSubsystem::GetAssetBeingOpenedMethod(TObjectPtr<UObject> Asset)
+{
+	if(!Asset)
+	{
+		return TOptional<EAssetOpenMethod>();
+	}
+	
+	EAssetOpenMethod* OpenMethod = AssetOpenMethodCache.Find(Asset->GetPathName());
+
+	return OpenMethod ? *OpenMethod : TOptional<EAssetOpenMethod>();
+}
+
+TOptional<EAssetOpenMethod> UAssetEditorSubsystem::GetAssetsBeingOpenedMethod(TArray<TObjectPtr<UObject>> Assets)
+{
+	TOptional<EAssetOpenMethod> FoundOpenMethod;
+
+	// If an asset editor supports opening multiple assets, if any of them are being opened in read only mode we ask the asset editor to open in read only mode
+	for(const TObjectPtr<UObject>& Asset : Assets)
+	{
+		TOptional<EAssetOpenMethod> AssetOpenMethod = GetAssetBeingOpenedMethod(Asset);
+
+		if(AssetOpenMethod.IsSet())
+		{
+			FoundOpenMethod = AssetOpenMethod.GetValue();
+
+			if(FoundOpenMethod == EAssetOpenMethod::View)
+			{
+				return FoundOpenMethod;
+			}
+		}
+	}
+
+	return FoundOpenMethod;
+}
+
+void UAssetEditorSubsystem::AddReadOnlyAssetFilter(const FName& Owner, const FReadOnlyAssetFilter& ReadOnlyAssetFilter)
+{
+	ReadOnlyAssetFilters.Add(Owner, ReadOnlyAssetFilter);
+}
+
+void UAssetEditorSubsystem::RemoveReadOnlyAssetFilter(const FName& Owner)
+{
+	ReadOnlyAssetFilters.Remove(Owner);
+}
+
+
 void UAssetEditorSubsystem::HandleRequestOpenAssetMessage(const FAssetEditorRequestOpenAsset& Message, const TSharedRef<IMessageContext, ESPMode::ThreadSafe>& Context)
 {
 	OpenEditorForAsset(Message.AssetName);
 }
 
-void UAssetEditorSubsystem::OpenEditorForAsset(const FSoftObjectPath& AssetPath, const EAssetTypeActivationOpenedMethod OpenedMethod)
+UObject* UAssetEditorSubsystem::FindOrLoadAssetForOpening(const FSoftObjectPath& AssetPath)
 {
 	UObject* Object = FindObject<UObject>(AssetPath.GetAssetPath());
 	if (!Object)
 	{
 		Object = LoadObject<UObject>(nullptr, *AssetPath.GetAssetPathString(), nullptr, LOAD_NoRedirects);
 	}
-	if (Object)
+
+	return Object;
+	
+}
+
+void UAssetEditorSubsystem::OpenEditorForAsset(const FSoftObjectPath& AssetPath, const EAssetTypeActivationOpenedMethod OpenedMethod)
+{
+	if (UObject* Object = FindOrLoadAssetForOpening(AssetPath))
 	{
 		OpenEditorForAsset(Object, EToolkitMode::Standalone, TSharedPtr<IToolkitHost>(), true, OpenedMethod);
 	}
@@ -1195,7 +1251,20 @@ void UAssetEditorSubsystem::CreateRecentAssetsMenu(UToolMenu* InMenu, const FNam
 		FUIAction(
 			FExecuteAction::CreateLambda([this, CurRecent]()
 				{
-					OpenEditorForAsset(CurRecent);
+					if (UObject* Object = FindOrLoadAssetForOpening(CurRecent))
+					{
+						FText ErrorMessage;
+
+						// Try to open the asset in edit mode. If that is not allowed, try to open it in read only mode
+						if(CanOpenEditorForAsset(Object, EAssetTypeActivationOpenedMethod::Edit, &ErrorMessage))
+						{
+							OpenEditorForAsset(Object, EToolkitMode::Standalone, TSharedPtr<IToolkitHost>(), true, EAssetTypeActivationOpenedMethod::Edit);
+						}
+						else if(CanOpenEditorForAsset(Object, EAssetTypeActivationOpenedMethod::View, &ErrorMessage))
+						{
+							OpenEditorForAsset(Object, EToolkitMode::Standalone, TSharedPtr<IToolkitHost>(), true, EAssetTypeActivationOpenedMethod::View);
+						}
+					}
 				})
 			)
 		);
@@ -1420,6 +1489,94 @@ void UAssetEditorSubsystem::OpenEditorsForAssets(const TArray<FName>& AssetsToOp
 	}
 }
 
+bool UAssetEditorSubsystem::CanOpenEditorForAsset(UObject* Asset, const EAssetTypeActivationOpenedMethod OpenedMethod, FText* OutErrorMsg)
+{
+	if (!Asset)
+	{
+		if(OutErrorMsg)
+		{
+			*OutErrorMsg = LOCTEXT("AssetNull", "Opening Asset editor failed because asset is null");
+		}
+		
+		return false;
+	}
+	
+	FAssetToolsModule& AssetToolsModule = FModuleManager::LoadModuleChecked<FAssetToolsModule>(TEXT("AssetTools"));
+	TWeakPtr<IAssetTypeActions> AssetTypeActions = AssetToolsModule.Get().GetAssetTypeActionsForClass(Asset->GetClass());
+
+	// First we check with the asset type action to see if it supports the open method
+	if(AssetTypeActions.IsValid())
+	{
+		if(!AssetTypeActions.Pin()->SupportsOpenedMethod(OpenedMethod))
+		{
+			if(OutErrorMsg)
+			{
+				if(OpenedMethod == EAssetTypeActivationOpenedMethod::Edit)
+				{
+					*OutErrorMsg = FText::Format(LOCTEXT("AssetTypeDoesntSupportEdit", "A {0} does not support being edited!"), AssetTypeActions.Pin()->GetName());
+				}
+				else
+				{
+					*OutErrorMsg = FText::Format(LOCTEXT("AssetTypeDoesntSupportReadOnly", "A {0} does not support being opened in read only mode!"), AssetTypeActions.Pin()->GetName());
+				}
+
+			}
+			return false;
+		}
+	}
+	else
+	{
+		// Disallow opening an asset editor for classes
+		if(Asset->IsA<UClass>())
+		{
+			if(OutErrorMsg)
+			{
+				*OutErrorMsg = LOCTEXT("AssetTypeDoesntSupportOpenMethod", "UClasses cannot be opened in Asset Editors!");
+			}
+			return false;
+		}
+	}
+
+	// If the asset needs to be edited, make sure that is possible
+	if(OpenedMethod == EAssetTypeActivationOpenedMethod::Edit)
+	{
+		if(!IsAssetEditable(Asset))
+		{
+			if(OutErrorMsg)
+			{
+				*OutErrorMsg = LOCTEXT("AssetTypeDoesntSupportOpenMethod", "Unable to Edit Cooked asset");
+			}
+			return false;
+		}
+	}
+
+	if(OpenedMethod == EAssetTypeActivationOpenedMethod::View)
+	{
+		if (UPackage* Package = Asset->GetPackage())
+		{
+			for(const TPair<FName, FReadOnlyAssetFilter>& Filter : ReadOnlyAssetFilters)
+			{
+				if(!Filter.Value.Execute(Package->GetPathName()))
+				{
+					if(OutErrorMsg)
+					{
+						*OutErrorMsg = LOCTEXT("AssetTypeDoesntSupportOpenMethod", "This asset does not support being opened in read only mode.");
+					}
+					return false;
+				}
+			}
+		}
+		else
+		{
+			return false; // failsafe, but the package should exist at this point
+		}
+
+		
+	}
+
+	return true; 
+}
+
 void UAssetEditorSubsystem::RegisterEditorModes()
 {
 	for (FThreadSafeObjectIterator EditorModeIter(UEdMode::StaticClass()); EditorModeIter; ++EditorModeIter)
diff --git a/Engine/Source/Editor/UnrealEd/Private/Toolkits/AssetEditorToolkit.cpp b/Engine/Source/Editor/UnrealEd/Private/Toolkits/AssetEditorToolkit.cpp
--- a/Engine/Source/Editor/UnrealEd/Private/Toolkits/AssetEditorToolkit.cpp
+++ b/Engine/Source/Editor/UnrealEd/Private/Toolkits/AssetEditorToolkit.cpp
@@ -35,6 +35,7 @@
 #include "Textures/SlateIcon.h"
 #include "WidgetDrawerConfig.h"
 #include "Interfaces/Interface_AsyncCompilation.h"
+#include "Widgets/Images/SImage.h"
 
 #define LOCTEXT_NAMESPACE "AssetEditorToolkit"
 
@@ -52,11 +53,12 @@ FAssetEditorToolkit::FAssetEditorToolkit()
 	, AssetEditorModeManager(nullptr)
 	, bIsToolbarFocusable(false)
 	, bIsToolbarUsingSmallIcons(false)
+	, OpenMethod(EAssetOpenMethod::Edit)
 {
 	WorkspaceMenuCategory = FWorkspaceItem::NewGroup(LOCTEXT("WorkspaceMenu_BaseAssetEditor", "Asset Editor"));
 }
 
-void FAssetEditorToolkit::InitAssetEditor( const EToolkitMode::Type Mode, const TSharedPtr< class IToolkitHost >& InitToolkitHost, const FName AppIdentifier, const TSharedRef<FTabManager::FLayout>& StandaloneDefaultLayout, const bool bCreateDefaultStandaloneMenu, const bool bCreateDefaultToolbar, UObject* ObjectToEdit, const bool bInIsToolbarFocusable, const bool bInUseSmallToolbarIcons )
+void FAssetEditorToolkit::InitAssetEditor( const EToolkitMode::Type Mode, const TSharedPtr< class IToolkitHost >& InitToolkitHost, const FName AppIdentifier, const TSharedRef<FTabManager::FLayout>& StandaloneDefaultLayout, const bool bCreateDefaultStandaloneMenu, const bool bCreateDefaultToolbar, UObject* ObjectToEdit, const bool bInIsToolbarFocusable, const bool bInUseSmallToolbarIcons, const TOptional<EAssetOpenMethod>& InOpenMethod )
 {
 	TArray< UObject* > ObjectsToEdit;
 	ObjectsToEdit.Add( ObjectToEdit );
@@ -64,7 +66,7 @@ void FAssetEditorToolkit::InitAssetEditor( const EToolkitMode::Type Mode, const
 	InitAssetEditor( Mode, InitToolkitHost, AppIdentifier, StandaloneDefaultLayout, bCreateDefaultStandaloneMenu, bCreateDefaultToolbar, ObjectsToEdit, bInIsToolbarFocusable, bInUseSmallToolbarIcons );
 }
 
-void FAssetEditorToolkit::InitAssetEditor( const EToolkitMode::Type Mode, const TSharedPtr< class IToolkitHost >& InitToolkitHost, const FName AppIdentifier, const TSharedRef<FTabManager::FLayout>& StandaloneDefaultLayout, const bool bCreateDefaultStandaloneMenu, const bool bCreateDefaultToolbar, const TArray<UObject*>& ObjectsToEdit, const bool bInIsToolbarFocusable, const bool bInUseSmallToolbarIcons )
+void FAssetEditorToolkit::InitAssetEditor( const EToolkitMode::Type Mode, const TSharedPtr< class IToolkitHost >& InitToolkitHost, const FName AppIdentifier, const TSharedRef<FTabManager::FLayout>& StandaloneDefaultLayout, const bool bCreateDefaultStandaloneMenu, const bool bCreateDefaultToolbar, const TArray<UObject*>& ObjectsToEdit, const bool bInIsToolbarFocusable, const bool bInUseSmallToolbarIcons, const TOptional<EAssetOpenMethod>& InOpenMethod )
 {
 	// Must not already be editing an object
 	check( ObjectsToEdit.Num() > 0 );
@@ -78,6 +80,22 @@ void FAssetEditorToolkit::InitAssetEditor( const EToolkitMode::Type Mode, const
 
 	EditingObjects.Append( ObjectsToEdit );
 
+	// If the open method was manually overriden, use that
+	if(InOpenMethod.IsSet())
+	{
+		OpenMethod = InOpenMethod.GetValue();
+	}
+	// Otherwise we check the Asset Editor Subsystem to see if there is a method this asset editor is being requested to open in
+	else
+	{
+		TOptional<EAssetOpenMethod> CachedOpenMethod =  GEditor->GetEditorSubsystem<UAssetEditorSubsystem>()->GetAssetsBeingOpenedMethod(EditingObjects);
+
+		if(CachedOpenMethod.IsSet())
+		{
+			OpenMethod = CachedOpenMethod.GetValue();
+		}
+	}
+
 	// Store "previous" asset editing toolkit host, and clear it out
 	PreviousWorldCentricToolkitHost = PreviousWorldCentricToolkitHostForNewAssetEditor;
 	PreviousWorldCentricToolkitHostForNewAssetEditor.Reset();
@@ -190,6 +208,7 @@ void FAssetEditorToolkit::InitAssetEditor( const EToolkitMode::Type Mode, const
 		const TSharedRef<FTabManager> NewTabManager = FGlobalTabmanager::Get()->NewTabManager( NewMajorTab.ToSharedRef() );		
 		NewTabManager->SetOnPersistLayout(FTabManager::FOnPersistLayout::CreateRaw(this, &FAssetEditorToolkit::HandleTabManagerPersistLayout));
 		NewTabManager->SetAllowWindowMenuBar(true);
+		NewTabManager->SetReadOnly(OpenMethod == EAssetOpenMethod::View);
 
 		this->TabManager = NewTabManager;
 
@@ -233,12 +252,12 @@ void FAssetEditorToolkit::InitAssetEditor( const EToolkitMode::Type Mode, const
 	ToolkitCommands->MapAction(
 		FAssetEditorCommonCommands::Get().SaveAsset,
 		FExecuteAction::CreateSP( this, &FAssetEditorToolkit::SaveAsset_Execute ),
-		FCanExecuteAction::CreateSP( this, &FAssetEditorToolkit::CanSaveAsset ));
+		FCanExecuteAction::CreateSP( this, &FAssetEditorToolkit::CanSaveAsset_Internal ));
 
 	ToolkitCommands->MapAction(
 		FAssetEditorCommonCommands::Get().SaveAssetAs,
 		FExecuteAction::CreateSP( this, &FAssetEditorToolkit::SaveAssetAs_Execute ),
-		FCanExecuteAction::CreateSP( this, &FAssetEditorToolkit::CanSaveAssetAs ));
+		FCanExecuteAction::CreateSP( this, &FAssetEditorToolkit::CanSaveAssetAs_Internal ));
 
 	ToolkitCommands->MapAction(
 		FGlobalEditorCommonCommands::Get().FindInContentBrowser,
@@ -254,7 +273,7 @@ void FAssetEditorToolkit::InitAssetEditor( const EToolkitMode::Type Mode, const
 	ToolkitCommands->MapAction(
 		FAssetEditorCommonCommands::Get().ReimportAsset,
 		FExecuteAction::CreateSP( this, &FAssetEditorToolkit::Reimport_Execute ),
-		FCanExecuteAction::CreateSP(this, &FAssetEditorToolkit::CanReimport));
+		FCanExecuteAction::CreateSP(this, &FAssetEditorToolkit::CanReimport_Internal));
 
 	FGlobalEditorCommonCommands::MapActions(ToolkitCommands);
 
@@ -315,7 +334,7 @@ void FAssetEditorToolkit::InitAssetEditor( const EToolkitMode::Type Mode, const
 	{
 		EditorModeManager->SetToolkitHost(ToolkitHost.Pin().ToSharedRef());
 	}
-
+	
 	PostInitAssetEditor();
 
 	// NOTE: Currently, the AssetEditorSubsystem will keep a hard reference to our object as we're editing it
@@ -343,7 +362,6 @@ FAssetEditorToolkit::~FAssetEditorToolkit()
 	PRAGMA_ENABLE_DEPRECATION_WARNINGS
 }
 
-
 void FAssetEditorToolkit::RegisterTabSpawners(const TSharedRef<class FTabManager>& InTabManager)
 {
 	// Use the first child category of the local workspace root if there is one, otherwise use the root itself
@@ -605,6 +623,15 @@ void FAssetEditorToolkit::RemoveEditingObject(UObject* Object)
 	GEditor->GetEditorSubsystem<UAssetEditorSubsystem>()->NotifyAssetClosed( Object, this );
 }
 
+bool FAssetEditorToolkit::CanSaveAsset_Internal() const
+{
+	if(GetOpenMethod() != EAssetOpenMethod::Edit)
+	{
+		return false;
+	}
+
+	return CanSaveAsset();
+}
 
 void FAssetEditorToolkit::SaveAsset_Execute()
 {
@@ -639,6 +666,15 @@ void FAssetEditorToolkit::SaveAsset_Execute()
 	FEditorFileUtils::PromptForCheckoutAndSave(PackagesToSave, bCheckDirtyOnAssetSave, /*bPromptToSave=*/ false);
 }
 
+bool FAssetEditorToolkit::CanSaveAssetAs_Internal() const
+{
+	if(GetOpenMethod() != EAssetOpenMethod::Edit)
+	{
+		return false;
+	}
+
+	return CanSaveAssetAs();
+}
 
 void FAssetEditorToolkit::SaveAssetAs_Execute()
 {
@@ -949,6 +985,15 @@ bool FAssetEditorToolkit::CanReimport() const
 	return false;
 }
 
+bool FAssetEditorToolkit::CanReimport_Internal() const
+{
+	if(GetOpenMethod() != EAssetOpenMethod::Edit)
+	{
+		return false;
+	}
+
+	return CanReimport();
+}
 
 bool FAssetEditorToolkit::CanReimport( UObject* EditingObject ) const
 {
@@ -1010,10 +1055,17 @@ void FAssetEditorToolkit::FillDefaultFileMenuCommands(FToolMenuSection& InSectio
 {
 	const FToolMenuInsert InsertPosition(NAME_None, EToolMenuInsertType::First);
 
-	InSection.AddMenuEntry(FAssetEditorCommonCommands::Get().SaveAsset, TAttribute<FText>(), TAttribute<FText>(), FSlateIcon(FAppStyle::GetAppStyleSetName(), "AssetEditor.SaveAsset")).InsertPosition = InsertPosition;
-	if( IsActuallyAnAsset() )
+	if (UAssetEditorToolkitMenuContext* Context = InSection.FindContext<UAssetEditorToolkitMenuContext>())
 	{
-		InSection.AddMenuEntry(FAssetEditorCommonCommands::Get().SaveAssetAs, TAttribute<FText>(), TAttribute<FText>(), FSlateIcon(FAppStyle::GetAppStyleSetName(), "AssetEditor.SaveAssetAs")).InsertPosition = InsertPosition;
+		// The save buttons are hidden if we are in read only mode
+		if(Context->Toolkit.Pin()->GetOpenMethod() == EAssetOpenMethod::Edit)
+		{
+			InSection.AddMenuEntry(FAssetEditorCommonCommands::Get().SaveAsset, TAttribute<FText>(), TAttribute<FText>(), FSlateIcon(FAppStyle::GetAppStyleSetName(), "AssetEditor.SaveAsset")).InsertPosition = InsertPosition;
+			if( IsActuallyAnAsset() )
+			{
+				InSection.AddMenuEntry(FAssetEditorCommonCommands::Get().SaveAssetAs, TAttribute<FText>(), TAttribute<FText>(), FSlateIcon(FAppStyle::GetAppStyleSetName(), "AssetEditor.SaveAssetAs")).InsertPosition = InsertPosition;
+			}
+		}
 	}
 
 	if( IsWorldCentricAssetEditor() )
@@ -1124,7 +1176,6 @@ void FAssetEditorToolkit::RegisterDefaultToolBar()
 
 void FAssetEditorToolkit::InitToolMenuContext(FToolMenuContext& MenuContext)
 {
-
 }
 
 UToolMenu* FAssetEditorToolkit::GenerateCommonActionsToolbar(FToolMenuContext& MenuContext)
@@ -1140,11 +1191,90 @@ UToolMenu* FAssetEditorToolkit::GenerateCommonActionsToolbar(FToolMenuContext& M
 		FoundMenu->StyleName = "AssetEditorToolbar";
 
 		FToolMenuSection& Section = FoundMenu->AddSection("CommonActions");
-		Section.AddEntry(FToolMenuEntry::InitToolBarButton(FAssetEditorCommonCommands::Get().SaveAsset));
+
+		Section.AddDynamicEntry("CommonActionsDynamic", FNewToolMenuSectionDelegate::CreateLambda([this](FToolMenuSection& InSection)
+		{
+			UAssetEditorToolkitMenuContext* AssetEditorToolkitMenuContext = InSection.FindContext<UAssetEditorToolkitMenuContext>();
+
+			if (AssetEditorToolkitMenuContext)
+			{
+				if(TSharedPtr<FAssetEditorToolkit> AssetEditorToolkit = AssetEditorToolkitMenuContext->Toolkit.Pin())
+				{
+					// The save button is hidden if we are in read only mode
+					if(AssetEditorToolkit->GetOpenMethod() == EAssetOpenMethod::Edit)
+					{
+						InSection.AddEntry(FToolMenuEntry::InitToolBarButton(FAssetEditorCommonCommands::Get().SaveAsset));
+					}
+				}
+			}
+		}));
+
 		Section.AddEntry(FToolMenuEntry::InitToolBarButton(FGlobalEditorCommonCommands::Get().FindInContentBrowser, LOCTEXT("FindInContentBrowserButton", "Browse")));
 		Section.AddSeparator(NAME_None);
 	}
 
+	return ToolMenus->GenerateMenu(ToolBarName, MenuContext);
+}
+
+UToolMenu* FAssetEditorToolkit::GenerateReadOnlyToolbar(FToolMenuContext& MenuContext)
+{
+	UToolMenus* ToolMenus = UToolMenus::Get();
+	FName ToolBarName = "AssetEditorToolbar.ReadOnly";
+
+	UToolMenu* FoundMenu = ToolMenus->FindMenu(ToolBarName);
+	
+	if (!FoundMenu || !FoundMenu->IsRegistered())
+	{
+		FoundMenu = ToolMenus->RegisterMenu(ToolBarName, NAME_None, EMultiBoxType::SlimHorizontalToolBar);
+		FoundMenu->StyleName = "AssetEditorToolbar";
+		FToolMenuSection& Section = FoundMenu->AddSection("ReadOnly");
+
+		Section.AddDynamicEntry("ReadOnlyDynamic", FNewToolMenuSectionDelegate::CreateLambda([this](FToolMenuSection& InSection)
+		{
+			// A custom widget to show the read only status of the asset editor
+			TSharedRef<SWidget> ReadOnlyIndicatorWidget =
+			SNew(SHorizontalBox)
+			+SHorizontalBox::Slot()
+			.Padding(20.0f, 0.0f)
+			[
+				SNew(SBorder)
+				.BorderImage(FAppStyle::GetBrush("AssetEditor.ReadOnlyBorder"))
+				.Padding(0.0f, 0.0f)
+				[
+					SNew(SHorizontalBox)
+					+SHorizontalBox::Slot()
+					.AutoWidth()
+					.HAlign(HAlign_Left)
+					.VAlign(VAlign_Center)
+					.Padding(8.0f, 0.0f, 0.0f, 0.0f)
+					[
+						SNew(SBox)
+						.HeightOverride(16.0f)
+						.WidthOverride(16.0f)
+						[
+							SNew(SImage)
+							.Image(FAppStyle::GetBrush("AssetEditor.ReadOnlyOpenable"))
+							.ColorAndOpacity(FStyleColors::AccentBlack)
+						]
+						
+					]
+					+SHorizontalBox::Slot()
+					.HAlign(HAlign_Left)
+					.VAlign(VAlign_Center)
+					.Padding(4.0f, 0.0f, 8.0f, 0.0f)
+					[
+						SNew(STextBlock)
+						.Text(LOCTEXT("ReadOnlyText", "Read Only"))
+						.ColorAndOpacity(FStyleColors::AccentBlack)
+					]
+				]	
+			];
+		
+			InSection.AddSeparator(NAME_None);
+			InSection.AddEntry(FToolMenuEntry::InitWidget("ReadOnlyIndicatorWidget", ReadOnlyIndicatorWidget, FText::GetEmpty()));
+		}));
+	}
+
 	return ToolMenus->GenerateMenu(ToolBarName, MenuContext);;
 }
 
@@ -1179,6 +1309,9 @@ void FAssetEditorToolkit::GenerateToolbar()
 	UToolMenu* CommonActionsToolbar = GenerateCommonActionsToolbar(MenuContext);
 	TSharedRef< class SWidget > CommonActionsToolbarWidget = ToolMenus->GenerateWidget(CommonActionsToolbar);
 
+	UToolMenu* ReadOnlyToolbar = GenerateReadOnlyToolbar(MenuContext);
+	TSharedRef< class SWidget > ReadOnlyWidget = ToolMenus->GenerateWidget(ReadOnlyToolbar);
+
 	TSharedRef<SWidget> MiscWidgets = SNullWidget::NullWidget;
 
 	if(ToolbarWidgets.Num() > 0)
@@ -1226,6 +1359,22 @@ void FAssetEditorToolkit::GenerateToolbar()
 			[
 				MiscWidgets
 			]
+		]
+		+SHorizontalBox::Slot()
+		.HAlign(HAlign_Right)
+		.AutoWidth()
+		[
+			SNew(SBorder)
+			.VAlign(VAlign_Center)
+			.Visibility_Lambda([this]()
+			{
+				return GetOpenMethod() == EAssetOpenMethod::View ? EVisibility::Visible : EVisibility::Collapsed;
+			})
+			.BorderImage(FAppStyle::Get().GetBrush("AssetEditorToolbar.Background"))
+			.Padding(FMargin(0.0f))
+			[
+				ReadOnlyWidget
+			]
 		];
 
 	if (ToolbarWidgetContent.IsValid())
diff --git a/Engine/Source/Editor/UnrealEd/Public/Subsystems/AssetEditorSubsystem.h b/Engine/Source/Editor/UnrealEd/Public/Subsystems/AssetEditorSubsystem.h
--- a/Engine/Source/Editor/UnrealEd/Public/Subsystems/AssetEditorSubsystem.h
+++ b/Engine/Source/Editor/UnrealEd/Public/Subsystems/AssetEditorSubsystem.h
@@ -11,6 +11,8 @@
 #include "Misc/NamePermissionList.h"
 #include "AssetTypeActivationOpenedMethod.h"
 #include "MRUFavoritesList.h"
+#include "AssetDefinition.h"
+
 #include "AssetEditorSubsystem.generated.h"
 
 class UAssetEditor;
@@ -77,6 +79,7 @@ public:
 	virtual TSharedPtr<class FTabManager> GetAssociatedTabManager() = 0;
 	virtual double GetLastActivationTime() = 0;
 	virtual void RemoveEditingAsset(UObject* Asset) = 0;
+	virtual EAssetOpenMethod GetOpenMethod() const { return EAssetOpenMethod::Edit; }
 };
 
 struct FRegisteredModeInfo
@@ -132,6 +135,9 @@ public:
 	UNREALED_API void OpenEditorsForAssets(const TArray<FName>& AssetsToOpen, const EAssetTypeActivationOpenedMethod OpenedMethod = EAssetTypeActivationOpenedMethod::Edit);
 	UNREALED_API void OpenEditorsForAssets(const TArray<FSoftObjectPath>& AssetsToOpen);
 
+	/** Check whether the given asset can be opened in the given open method */
+	UNREALED_API bool CanOpenEditorForAsset(UObject* Asset, const EAssetTypeActivationOpenedMethod OpenedMethod, FText* OutErrorMsg); 
+
 	/** Returns the primary editor if one is already open for the specified asset.
 	 * If there is one open and bFocusIfOpen is true, that editor will be brought to the foreground and focused if possible.
 	 */
@@ -311,6 +317,20 @@ public:
 
 	UNREALED_API void SetRecentAssetsFilter(const FMainMRUFavoritesList::FDoesMRUFavoritesItemPassFilter& InFilter);
 
+	/** These functions are used by the Asset Editor Toolkit to query the method in which the given assets are being opened
+	 *  These are only valid when the asset is in the process of being opened, i.e while we are in OpenEditorForAsset.
+	 *  After which, you should query the asset editor itself to get the open method
+	 */
+	TOptional<EAssetOpenMethod> GetAssetBeingOpenedMethod(TObjectPtr<UObject> Asset);
+	TOptional<EAssetOpenMethod> GetAssetsBeingOpenedMethod(TArray<TObjectPtr<UObject>> Assets);
+
+	/* Functionality to add a filter that is run to check if an asset is allowed to be opened in read only mode
+	 * (If the asset type supports it)
+	 */
+	DECLARE_DELEGATE_RetVal_OneParam(bool, FReadOnlyAssetFilter, const FString&);
+	UNREALED_API void AddReadOnlyAssetFilter(const FName& Owner, const FReadOnlyAssetFilter& ReadOnlyAssetFilter);
+	UNREALED_API void RemoveReadOnlyAssetFilter(const FName& Owner);
+
 private:
 
 	/** Handles a package being reloaded */
@@ -332,6 +352,8 @@ private:
 	UNREALED_API bool ShouldShowRecentAsset(const FString& AssetName, int32 RecentAssetIndex,  const FName& InAssetEditorName) const;
 	UNREALED_API bool ShouldShowRecentAssetsMenu(const FName& InAssetEditorName) const;
 
+	UNREALED_API UObject* FindOrLoadAssetForOpening(const FSoftObjectPath& AssetPath);
+
 private:
 
 	/** struct used by OpenedEditorTimes map to store editor names and times */
@@ -463,4 +485,12 @@ private:
 
 	/** The max number of recent assets to show in the menu */
 	int32 MaxRecentAssetsToShowInMenu = 20;
+
+	/** The method any assets being opened are requested to open in, only valid during the open process (OpenEditorForAsset)
+	 *  Since there is no way to generically get this information to FAssetEditorToolkit::InitAssetEditor
+	 */
+	TMap<FString, EAssetOpenMethod> AssetOpenMethodCache;
+
+	/** External filters that are run to check if an asset is openable in read only mode */
+	TMap<FName, FReadOnlyAssetFilter> ReadOnlyAssetFilters;
 };
diff --git a/Engine/Source/Editor/UnrealEd/Public/Toolkits/AssetEditorToolkit.h b/Engine/Source/Editor/UnrealEd/Public/Toolkits/AssetEditorToolkit.h
--- a/Engine/Source/Editor/UnrealEd/Public/Toolkits/AssetEditorToolkit.h
+++ b/Engine/Source/Editor/UnrealEd/Public/Toolkits/AssetEditorToolkit.h
@@ -7,6 +7,7 @@
 #include "UObject/GCObject.h"
 #include "Framework/Docking/TabManager.h"
 #include "Toolkits/IToolkit.h"
+#include "AssetDefinition.h"
 
 #include "Framework/Commands/UICommandList.h"
 #include "Framework/MultiBox/MultiBoxExtender.h"
@@ -114,9 +115,10 @@ public:
 	 * @param	ObjectToEdit			The object to edit
 	 * @param	bInIsToolbarFocusable	Whether the buttons on the default toolbar can receive keyboard focus
 	 * @param	bUseSmallToolbarIcons	Whether the buttons on the default toolbar use the small icons
+	 * @param	InOpenMethod			Override whether the Asset Editor is being opened in read only or edit mode (otherwise automatically set by the asset editor subsytem for any asset editors opened through it)
 	 */
-	UNREALED_API void InitAssetEditor(const EToolkitMode::Type Mode, const TSharedPtr<IToolkitHost>& InitToolkitHost, const FName AppIdentifier, const TSharedRef<FTabManager::FLayout>& StandaloneDefaultLayout, const bool bCreateDefaultStandaloneMenu, const bool bCreateDefaultToolbar, const TArray<UObject*>& ObjectsToEdit, const bool bInIsToolbarFocusable = false, const bool bInUseSmallToolbarIcons = false);
-	UNREALED_API void InitAssetEditor(const EToolkitMode::Type Mode, const TSharedPtr<IToolkitHost>& InitToolkitHost, const FName AppIdentifier, const TSharedRef<FTabManager::FLayout>& StandaloneDefaultLayout, const bool bCreateDefaultStandaloneMenu, const bool bCreateDefaultToolbar, UObject* ObjectToEdit, const bool bInIsToolbarFocusable = false, const bool bInUseSmallToolbarIcons = false);
+	UNREALED_API void InitAssetEditor(const EToolkitMode::Type Mode, const TSharedPtr<IToolkitHost>& InitToolkitHost, const FName AppIdentifier, const TSharedRef<FTabManager::FLayout>& StandaloneDefaultLayout, const bool bCreateDefaultStandaloneMenu, const bool bCreateDefaultToolbar, const TArray<UObject*>& ObjectsToEdit, const bool bInIsToolbarFocusable = false, const bool bInUseSmallToolbarIcons = false, const TOptional<EAssetOpenMethod>& InOpenMethod = TOptional<EAssetOpenMethod>());
+	UNREALED_API void InitAssetEditor(const EToolkitMode::Type Mode, const TSharedPtr<IToolkitHost>& InitToolkitHost, const FName AppIdentifier, const TSharedRef<FTabManager::FLayout>& StandaloneDefaultLayout, const bool bCreateDefaultStandaloneMenu, const bool bCreateDefaultToolbar, UObject* ObjectToEdit, const bool bInIsToolbarFocusable = false, const bool bInUseSmallToolbarIcons = false, const TOptional<EAssetOpenMethod>& InOpenMethod = TOptional<EAssetOpenMethod>());
 
 	FAssetEditorToolkit(const FAssetEditorToolkit&) = delete;
 	FAssetEditorToolkit& operator=(const FAssetEditorToolkit&) = delete;
@@ -149,6 +151,7 @@ public:
 	UNREALED_API virtual TSharedPtr<FTabManager> GetAssociatedTabManager() override;
 	UNREALED_API virtual double GetLastActivationTime() override;
 	UNREALED_API virtual void RemoveEditingAsset(UObject* Asset) override;
+	UNREALED_API virtual EAssetOpenMethod GetOpenMethod() const override { return OpenMethod; }
 
 	/**
 	 * Fills in the supplied menu with commands for working with this asset file
@@ -305,6 +308,9 @@ protected:
 	/** Generate the toolbar for common asset actions like Save*/
 	UNREALED_API UToolMenu* GenerateCommonActionsToolbar(FToolMenuContext& MenuContext);
 
+	/** Generate the toolbar for read only mode specific content */
+	UToolMenu* GenerateReadOnlyToolbar(FToolMenuContext& MenuContext);
+
 	/** Get the collection of edited objects that can be saved. */
 	UNREALED_API virtual void GetSaveableObjects(TArray<UObject*>& OutObjects) const;
 
@@ -317,12 +323,18 @@ protected:
 	/** Called to test if "Save" should be enabled for this asset */
 	virtual bool CanSaveAsset() const { return true; }
 
+	/** Internal function to check if the asset can be saved that calls CanSaveAsset */
+	bool CanSaveAsset_Internal() const;
+
 	/** Called when "Save" is clicked for this asset */
 	UNREALED_API virtual void SaveAsset_Execute();
 
 	/** Called to test if "Save As" should be enabled for this asset */
 	virtual bool CanSaveAssetAs() const { return true; }
 
+	/** Internal function to check if the asset can be saved that calls CanSaveAssetAs */
+	bool CanSaveAssetAs_Internal() const;
+
 	/** Called when "Save As" is clicked for this asset */
 	UNREALED_API virtual void SaveAssetAs_Execute();
 
@@ -349,8 +361,10 @@ protected:
 
 	/** Called to check to see if there's an asset capable of being reimported */
 	UNREALED_API virtual bool CanReimport() const;
+	/** Internal function to check if the asset can be reimported that calls CanReimport*/
+	bool CanReimport_Internal() const;
 	UNREALED_API virtual bool CanReimport(UObject* EditingObject) const;
-
+	
 	/** Called when "Reimport" is clicked for this asset */
 	UNREALED_API virtual void Reimport_Execute();
 	UNREALED_API virtual void Reimport_Execute(UObject* EditingObject);
@@ -475,4 +489,7 @@ private:
 
 	/** Whether the buttons on the default toolbar use small icons */
 	bool bIsToolbarUsingSmallIcons;
+
+	/** Whether the asset editor was opened in edit mode or read only mode */
+	EAssetOpenMethod OpenMethod;
 };
diff --git a/Engine/Source/Runtime/Slate/Private/Framework/Docking/SDockingTabStack.cpp b/Engine/Source/Runtime/Slate/Private/Framework/Docking/SDockingTabStack.cpp
--- a/Engine/Source/Runtime/Slate/Private/Framework/Docking/SDockingTabStack.cpp
+++ b/Engine/Source/Runtime/Slate/Private/Framework/Docking/SDockingTabStack.cpp
@@ -179,6 +179,7 @@ void SDockingTabStack::Construct( const FArguments& InArgs, const TSharedRef<FTa
 				.BorderImage(this, &SDockingTabStack::GetContentAreaBrush)
 				.Padding(this, &SDockingTabStack::GetContentPadding)
 				.Clipping(EWidgetClipping::ClipToBounds)
+				.IsEnabled(this, &SDockingTabStack::IsContentEnabled)
 				[
 					SNew(STextBlock)
 					.Text(LOCTEXT("EmptyTabMessage", "Empty Tab!"))
@@ -922,6 +923,25 @@ void SDockingTabStack::SetParentNode( TSharedRef<class SDockingSplitter> InParen
 	TitleBarSlot->AttachWidget(TitleBarContent.ToSharedRef());
 }
 
+bool SDockingTabStack::IsContentEnabled() const
+{
+	TSharedRef<FTabManager> TabManager = GetDockArea()->GetTabManager();
+
+	if(!TabManager->IsReadOnly())
+	{
+		return true;
+	}
+
+	// If we are in read only mode, and the foreground tab desires custom behavior (i.e not hidden or disabled) it is enabled
+	// and the tab owner is responsible for handling the content in read only mode
+	if(TSharedPtr<SDockTab> ForegroundTab = TabWell->GetForegroundTab())
+	{
+		return TabManager->GetTabReadOnlyBehavior(ForegroundTab->GetLayoutIdentifier()) == ETabReadOnlyBehavior::Custom;
+	}
+
+	return true;
+}
+
 
 /** What should the content area look like for the current tab? */
 const FSlateBrush* SDockingTabStack::GetContentAreaBrush() const
diff --git a/Engine/Source/Runtime/Slate/Private/Framework/Docking/SDockingTabStack.h b/Engine/Source/Runtime/Slate/Private/Framework/Docking/SDockingTabStack.h
--- a/Engine/Source/Runtime/Slate/Private/Framework/Docking/SDockingTabStack.h
+++ b/Engine/Source/Runtime/Slate/Private/Framework/Docking/SDockingTabStack.h
@@ -204,6 +204,9 @@ protected:
 	/** Overridden from SDockingNode */
 	SLATE_API virtual void SetParentNode( TSharedRef<class SDockingSplitter> InParent ) override;
 
+	/** @return false if the content is disabled (non-interactable) */
+	bool IsContentEnabled() const;
+
 private:
 	
 	/** Data that persists across sessions and when the widget associated with this node is removed. */	
diff --git a/Engine/Source/Runtime/Slate/Private/Framework/Docking/TabManager.cpp b/Engine/Source/Runtime/Slate/Private/Framework/Docking/TabManager.cpp
--- a/Engine/Source/Runtime/Slate/Private/Framework/Docking/TabManager.cpp
+++ b/Engine/Source/Runtime/Slate/Private/Framework/Docking/TabManager.cpp
@@ -851,6 +851,20 @@ void FTabManager::SetMainTab(const TSharedRef<const SDockTab>& InTab)
 	
 }
 
+void FTabManager::SetReadOnly(bool bInReadOnly)
+{
+	if(bReadOnly != bInReadOnly)
+	{
+		bReadOnly = bInReadOnly;
+		OnReadOnlyModeChanged.Broadcast(bReadOnly);
+	}
+}
+
+bool FTabManager::IsReadOnly()
+{
+	return bReadOnly;
+}
+
 bool FTabManager::IsTabCloseable(const TSharedRef<const SDockTab>& InTab) const
 {
 	return MainNonCloseableTabID != InTab->GetLayoutIdentifier();
@@ -1172,7 +1186,8 @@ void FTabManager::PopulateTabSpawnerMenu_Helper( FMenuBuilder& PopulateMe, FPopu
 
 void FTabManager::MakeSpawnerMenuEntry( FMenuBuilder &PopulateMe, const TSharedPtr<FTabSpawnerEntry> &InSpawnerNode ) 
 {
-	if (InSpawnerNode->MenuType.Get() != ETabSpawnerMenuType::Hidden )
+	// We don't want to add a menu entry for this tab if it is hidden, or if we are in read only mode and it is asking to be hidden
+	if (InSpawnerNode->MenuType.Get() != ETabSpawnerMenuType::Hidden && !(bReadOnly && InSpawnerNode->ReadOnlyBehavior == ETabReadOnlyBehavior::Hidden) )
 	{
 		PopulateMe.AddMenuEntry(
 			InSpawnerNode->GetDisplayName().IsEmpty() ? FText::FromName(InSpawnerNode->TabType ) : InSpawnerNode->GetDisplayName(),
@@ -1940,7 +1955,31 @@ bool FTabManager::IsValidTabForSpawning( const FTab& SomeTab ) const
 
 bool FTabManager::IsAllowedTab(const FTabId& TabId) const
 {
-	return IsAllowedTabType(TabId.TabType);
+	bool bAllowed = true;
+
+	// If we are in read-only mode, make sure this tab doesn't want to be hidden
+	if(bReadOnly)
+	{
+		TOptional<ETabReadOnlyBehavior> TabReadOnlyBehavior = GetTabReadOnlyBehavior(TabId);
+
+		if(TabReadOnlyBehavior.IsSet())
+		{
+			bAllowed &= (TabReadOnlyBehavior.GetValue() != ETabReadOnlyBehavior::Hidden);
+		}
+	}
+	
+	bAllowed &= IsAllowedTabType(TabId.TabType);
+	
+	return bAllowed;
+}
+
+TOptional<ETabReadOnlyBehavior> FTabManager::GetTabReadOnlyBehavior(const FTabId& TabId) const
+{
+	if (const TSharedPtr<const FTabSpawnerEntry> Spawner = FindTabSpawnerFor(TabId.TabType))
+	{
+		return Spawner->ReadOnlyBehavior;
+	}
+	return TOptional<ETabReadOnlyBehavior>();
 }
 
 bool FTabManager::IsAllowedTabType(const FName TabType) const
diff --git a/Engine/Source/Runtime/Slate/Public/Framework/Docking/TabManager.h b/Engine/Source/Runtime/Slate/Public/Framework/Docking/TabManager.h
--- a/Engine/Source/Runtime/Slate/Public/Framework/Docking/TabManager.h
+++ b/Engine/Source/Runtime/Slate/Public/Framework/Docking/TabManager.h
@@ -202,6 +202,16 @@ namespace ETabSpawnerMenuType
 	};
 }
 
+/** An enum to describe how TabSpawnerEntries behave when the tab manager is in read only mode */
+enum class ETabReadOnlyBehavior : int32
+{
+	Disabled, // Default behavior - This tab will show up in read only modes but the contents will be disabled
+	Hidden,   // This tab will not show up in read only modes
+	Custom    // This tab will show up in read only modes and the tab owner will decide how it behaves (FTabManager::OnReadOnlyStateChanged to keep track of read only state)
+};
+
+DECLARE_MULTICAST_DELEGATE_OneParam(FOnReadOnlyModeChanged, bool /*bReadOnly*/);
+
 struct FTabSpawnerEntry : public FWorkspaceItem
 {
 	FTabSpawnerEntry(const FName& InTabType, const FOnSpawnTab& InSpawnTabMethod, const FCanSpawnTab& InCanSpawnTab)
@@ -214,9 +224,17 @@ struct FTabSpawnerEntry : public FWorkspaceItem
 		, bAutoGenerateMenuEntry(true)
 		, bCanSidebarTab(true)
 		, SpawnedTabPtr()
+		, ReadOnlyBehavior(ETabReadOnlyBehavior::Disabled)
+	
 	{
 	}
 
+	FTabSpawnerEntry& SetReadOnlyBehavior( const ETabReadOnlyBehavior& InReadOnlyBehavior)
+	{
+		ReadOnlyBehavior = InReadOnlyBehavior;
+		return *this;
+	}
+	
 	FTabSpawnerEntry& SetIcon( const FSlateIcon& InIcon)
 	{
 		Icon = InIcon;
@@ -312,6 +330,9 @@ private:
 
 	TWeakPtr<SDockTab> SpawnedTabPtr;
 
+	/** How this tab behaves when the tab manager is in a read only mode */
+	ETabReadOnlyBehavior ReadOnlyBehavior;
+
 	FORCENOINLINE bool IsSoleTabInstanceSpawned() const
 	{
 		// Items that allow multiple instances need a custom way to find and reuse tabs.
@@ -945,9 +966,15 @@ class FTabManager : public TSharedFromThis<FTabManager>
 		 */
 		SLATE_API void SetMainTab(const TSharedRef<const SDockTab>& InTab);
 
-		/** Provide a tab that will be the main tab and cannot be closed. */
 		SLATE_API void SetMainTab(const FTabId& InMainTabID);
+
+		/** Is this Tab Manager in Read Only mode i.e all interactions with panels are disabled */
+		SLATE_API bool IsReadOnly();
+
+		SLATE_API void SetReadOnly(bool bInReadOnly);
 	
+		SLATE_API FOnReadOnlyModeChanged& GetOnReadOnlyModeChangedDelegate() { return OnReadOnlyModeChanged; }
+
 		/* Prevent or allow all tabs to be drag */
 		void SetCanDoDragOperation(bool CanDoDragOperation) { bCanDoDragOperation = CanDoDragOperation; }
 		
@@ -960,6 +987,9 @@ class FTabManager : public TSharedFromThis<FTabManager>
 		/** @return true if a tab is ever allowed in a sidebar */
 		SLATE_API bool IsTabAllowedInSidebar(const FTabId TabId) const;
 
+		/** @return how the given tab wants to behave in read only mode */
+		TOptional<ETabReadOnlyBehavior> GetTabReadOnlyBehavior(const FTabId& TabId) const;
+
 		/**
 		 * Temporarily moves all open tabs in this tab manager to a sidebar or restores them from a temporary state
 		 */
@@ -1189,6 +1219,12 @@ class FTabManager : public TSharedFromThis<FTabManager>
 
 		/** Tabs which have been temporarily put in the a sidebar */
 		TArray<TWeakPtr<SDockTab>> TemporarilySidebaredTabs;
+
+		/** Whether this tab manager is in a read only mode (all tabs content disabled by default - overridable per tab) */
+		bool bReadOnly = false;
+
+		/** Delegate that broadcasts when the tab manager enters/leaves read only mode */
+		FOnReadOnlyModeChanged OnReadOnlyModeChanged;
 };
 
 
