Fix crash when clicking outside the mesh in the static mesh editor viewport if meshes had more than one section.

The edge selection code wasn't taking into account the start index for each section. It assumed that sections are arranged consecutively in the index buffer, which is not true in general.

#rnx

diff --git a/Engine/Source/Editor/StaticMeshEditor/Private/StaticMeshEditorViewportClient.cpp b/Engine/Source/Editor/StaticMeshEditor/Private/StaticMeshEditorViewportClient.cpp
--- a/Engine/Source/Editor/StaticMeshEditor/Private/StaticMeshEditorViewportClient.cpp
+++ b/Engine/Source/Editor/StaticMeshEditor/Private/StaticMeshEditorViewportClient.cpp
@@ -1114,12 +1114,12 @@ void FStaticMeshEditorViewportClient::ProcessClick(class FSceneView& InView, cla
 					FStaticMeshLODResources& RenderData = StaticMeshComponent->GetStaticMesh()->GetRenderData()->LODResources[LODLevel];
 
 					int32 NumBackFacingTriangles = 0;
-					uint32 IndexBufferIndex = 0;
 					for (int32 SectionIndex = 0; SectionIndex < RenderData.Sections.Num(); ++SectionIndex)
 					{
 						const FStaticMeshSection& Section = RenderData.Sections[SectionIndex];
 						const int32 FaceMaterialIndex = Section.MaterialIndex;
 						const int32 NumFaces = Section.NumTriangles;
+						uint32 IndexBufferIndex = Section.FirstIndex;
 						for (int32 FaceIndex = 0; FaceIndex < NumFaces; ++FaceIndex)
 						{
 							FVector VertexPosition[3];
