Add level set visualization to static mesh editor's simple collision visualization + editing system

diff --git a/Engine/Source/Editor/StaticMeshEditor/Private/StaticMeshEditor.cpp b/Engine/Source/Editor/StaticMeshEditor/Private/StaticMeshEditor.cpp
--- a/Engine/Source/Editor/StaticMeshEditor/Private/StaticMeshEditor.cpp
+++ b/Engine/Source/Editor/StaticMeshEditor/Private/StaticMeshEditor.cpp
@@ -1076,6 +1076,8 @@ bool FStaticMeshEditor::IsPrimValid(const FPrimData& InPrimData) const
 			return AggGeom->SphylElems.IsValidIndex(InPrimData.PrimIndex);
 		case EAggCollisionShape::Convex:
 			return AggGeom->ConvexElems.IsValidIndex(InPrimData.PrimIndex);
+		case EAggCollisionShape::LevelSet:
+			return AggGeom->LevelSetElems.IsValidIndex(InPrimData.PrimIndex);
 		}
 	}
 	return false;
@@ -1176,6 +1178,12 @@ void FStaticMeshEditor::DuplicateSelectedPrims(const FVector* InOffset)
 					PrimData.PrimIndex = AggGeom->ConvexElems.Add(ConvexElem);
 				}
 				break;
+			case EAggCollisionShape::LevelSet:
+				{
+					const FKLevelSetElem LevelSetElem = AggGeom->LevelSetElems[PrimData.PrimIndex];
+					PrimData.PrimIndex = AggGeom->LevelSetElems.Add(LevelSetElem);
+				}
+				break;
 			}
 
 			// If specified, offset the duplicate by a specific amount
@@ -1379,6 +1387,14 @@ void FStaticMeshEditor::ScaleSelectedPrims(const FVector& InScale)
 		case EAggCollisionShape::Convex:
 			AggGeom->ConvexElems[PrimData.PrimIndex].ScaleElem(ModifiedScale, MinPrimSize);
 			break;
+		case EAggCollisionShape::LevelSet:
+			{
+				// Apply scaling to the centered transform; note that MinPrimSize has no effect for level sets (nor convex hulls)
+				FTransform ScaledTransform = AggGeom->LevelSetElems[PrimData.PrimIndex].GetCenteredTransform();
+				ScaledTransform.SetScale3D(ScaledTransform.GetScale3D() + ModifiedScale);
+				AggGeom->LevelSetElems[PrimData.PrimIndex].SetCenteredTransform(ScaledTransform);
+				break;
+			}
 		}
 
 		StaticMesh->bCustomizedCollision = true;	//mark the static mesh for collision customization
@@ -1413,6 +1429,9 @@ bool FStaticMeshEditor::CalcSelectedPrimsAABB(FBox &OutBox) const
 		case EAggCollisionShape::Convex:
 			OutBox += AggGeom->ConvexElems[PrimData.PrimIndex].CalcAABB(FTransform::Identity, FVector(1.f));
 			break;
+		case EAggCollisionShape::LevelSet:
+			OutBox += AggGeom->LevelSetElems[PrimData.PrimIndex].CalcAABB(FTransform::Identity, FVector(1.f));
+			break;
 		}
 	}
 	return HasSelectedPrims();
@@ -1451,6 +1470,9 @@ bool FStaticMeshEditor::GetLastSelectedPrimTransform(FTransform& OutTransform) c
 		case EAggCollisionShape::Convex:
 			OutTransform = AggGeom->ConvexElems[PrimData.PrimIndex].GetTransform();
 			break;
+		case EAggCollisionShape::LevelSet:
+			OutTransform = AggGeom->LevelSetElems[PrimData.PrimIndex].GetCenteredTransform();
+			break;
 		}
 	}
 	return HasSelectedPrims();
@@ -1473,6 +1495,8 @@ FTransform FStaticMeshEditor::GetPrimTransform(const FPrimData& InPrimData) cons
 		return AggGeom->SphylElems[InPrimData.PrimIndex].GetTransform();
 	case EAggCollisionShape::Convex:
 		return AggGeom->ConvexElems[InPrimData.PrimIndex].GetTransform();
+	case EAggCollisionShape::LevelSet:
+		return AggGeom->LevelSetElems[InPrimData.PrimIndex].GetCenteredTransform();
 	}
 	return FTransform::Identity;
 }
@@ -1498,6 +1522,9 @@ void FStaticMeshEditor::SetPrimTransform(const FPrimData& InPrimData, const FTra
 	case EAggCollisionShape::Convex:
 		AggGeom->ConvexElems[InPrimData.PrimIndex].SetTransform(InPrimTransform);
 		break;
+	case EAggCollisionShape::LevelSet:
+		AggGeom->LevelSetElems[InPrimData.PrimIndex].SetCenteredTransform(InPrimTransform);
+		break;
 	}
 
 	StaticMesh->bCustomizedCollision = true;	//mark the static mesh for collision customization
@@ -1593,6 +1620,26 @@ bool FStaticMeshEditor::OverlapsExistingPrim(const FPrimData& InPrimData) const
 			}
 		}
 		break;
+	case EAggCollisionShape::LevelSet:
+		{
+			const FKLevelSetElem InLevelSetElem = AggGeom->LevelSetElems[InPrimData.PrimIndex];
+			const FTransform InElemTM = InLevelSetElem.GetTransform();
+			for (int32 i = 0; i < AggGeom->LevelSetElems.Num(); ++i)
+			{
+				if (i == InPrimData.PrimIndex)
+				{
+					continue;
+				}
+
+				const FKLevelSetElem& LevelSetElem = AggGeom->LevelSetElems[i];
+				const FTransform ElemTM = LevelSetElem.GetTransform();
+				if (InElemTM.Equals(ElemTM))
+				{
+					return true;
+				}
+			}
+		}
+		break;
 	}
 
 	return false;
@@ -2386,6 +2433,9 @@ void FStaticMeshEditor::DeleteSelectedPrims()
 			case EAggCollisionShape::Convex:
 				AggGeom->ConvexElems.RemoveAt(PrimData.PrimIndex);
 				break;
+			case EAggCollisionShape::LevelSet:
+				AggGeom->LevelSetElems.RemoveAt(PrimData.PrimIndex);
+				break;
 			}
 		}
 
diff --git a/Engine/Source/Editor/StaticMeshEditor/Private/StaticMeshEditorViewportClient.cpp b/Engine/Source/Editor/StaticMeshEditor/Private/StaticMeshEditorViewportClient.cpp
--- a/Engine/Source/Editor/StaticMeshEditor/Private/StaticMeshEditorViewportClient.cpp
+++ b/Engine/Source/Editor/StaticMeshEditor/Private/StaticMeshEditorViewportClient.cpp
@@ -566,6 +566,19 @@ void FStaticMeshEditorViewportClient::Draw(const FSceneView* View,FPrimitiveDraw
 
 			PDI->SetHitProxy(NULL);
 		}
+
+		for (int32 i = 0; i < AggGeom->LevelSetElems.Num(); ++i)
+		{
+			HSMECollisionProxy* HitProxy = new HSMECollisionProxy(EAggCollisionShape::LevelSet, i);
+			PDI->SetHitProxy(HitProxy);
+
+			const FColor CollisionColor = StaticMeshEditor->IsSelectedPrim(HitProxy->PrimData) ? SelectedColor : UnselectedColor;
+			const FKLevelSetElem& LevelSetElem = AggGeom->LevelSetElems[i];
+			const FTransform ElemTM = LevelSetElem.GetTransform();
+			LevelSetElem.DrawElemWire(PDI, ElemTM, 1.f, CollisionColor);
+
+			PDI->SetHitProxy(NULL);
+		}
 	}
 
 	if (bShowComplexCollision && StaticMesh->ComplexCollisionMesh && StaticMesh->GetBodySetup()->CollisionTraceFlag != ECollisionTraceFlag::CTF_UseSimpleAsComplex)
diff --git a/Engine/Source/Runtime/Engine/Classes/PhysicsEngine/LevelSetElem.h b/Engine/Source/Runtime/Engine/Classes/PhysicsEngine/LevelSetElem.h
--- a/Engine/Source/Runtime/Engine/Classes/PhysicsEngine/LevelSetElem.h
+++ b/Engine/Source/Runtime/Engine/Classes/PhysicsEngine/LevelSetElem.h
@@ -47,6 +47,25 @@ struct FKLevelSetElem : public FKShapeElem
 		Transform = InTransform;
 	}
 
+	// Get the tranform of the center of the level set. In contrast, GetTransform() gets a transform relative to the corner.
+	FTransform GetCenteredTransform() const
+	{
+		return FTransform(
+			Transform.GetRotation(),
+			Transform.GetLocation() + Transform.TransformVector(UntransformedAABB().GetCenter()),
+			Transform.GetScale3D());
+	}
+
+	// Set the tranform of the center of the level set. In contrast, SetTransform() sets a transform relative to the corner.
+	void SetCenteredTransform(const FTransform& CenteredTransform)
+	{
+		ensure(CenteredTransform.IsValid());
+		Transform = FTransform(
+			CenteredTransform.GetRotation(),
+			CenteredTransform.GetLocation() - CenteredTransform.TransformVector(UntransformedAABB().GetCenter()),
+			CenteredTransform.GetScale3D());
+	}
+
 	ENGINE_API void ScaleElem(FVector DeltaSize, float MinSize);
 
 	// Draw helpers
