Remove needless include of NaniteDefinitions.h to cut UE build items due to changes from ~5000 to ~100

#rnx

diff --git a/Engine/Source/Editor/StaticMeshEditor/Private/StaticMeshEditorTools.cpp b/Engine/Source/Editor/StaticMeshEditor/Private/StaticMeshEditorTools.cpp
--- a/Engine/Source/Editor/StaticMeshEditor/Private/StaticMeshEditorTools.cpp
+++ b/Engine/Source/Editor/StaticMeshEditor/Private/StaticMeshEditorTools.cpp
@@ -51,6 +51,7 @@
 #include "PerPlatformPropertyCustomization.h"
 #include "Misc/ScopedSlowTask.h"
 #include "MeshCardRepresentation.h"
+#include "NaniteDefinitions.h"
 
 const uint32 MaxHullCount = 64;
 const uint32 MinHullCount = 1;
diff --git a/Engine/Source/Runtime/Engine/Public/Components.h b/Engine/Source/Runtime/Engine/Public/Components.h
--- a/Engine/Source/Runtime/Engine/Public/Components.h
+++ b/Engine/Source/Runtime/Engine/Public/Components.h
@@ -7,8 +7,9 @@
 #include "Containers/StridedView.h"
 #include "MeshUVChannelInfo.h"
 #include "VertexStreamComponent.h"
+#if UE_ENABLE_INCLUDE_ORDER_DEPRECATED_IN_5_4
 #include "NaniteDefinitions.h"
-
+#endif
 #if UE_ENABLE_INCLUDE_ORDER_DEPRECATED_IN_5_2
 #include "RenderResource.h"
 #include "VertexFactory.h"
