Make the static mesh editor create its socketmanager on post-init, so it's there when user interacts with sockets on the mesh (even if the sockets window was never opened)

#changelist validated
#virtualized

diff --git a/Engine/Source/Editor/StaticMeshEditor/Private/StaticMeshEditor.cpp b/Engine/Source/Editor/StaticMeshEditor/Private/StaticMeshEditor.cpp
--- a/Engine/Source/Editor/StaticMeshEditor/Private/StaticMeshEditor.cpp
+++ b/Engine/Source/Editor/StaticMeshEditor/Private/StaticMeshEditor.cpp
@@ -391,6 +391,9 @@ void FStaticMeshEditor::PostInitAssetEditor()
 		FStaticMeshEditorViewportClient& ViewportClient = StaticMeshViewport->GetViewportClient();
 		ViewportClient.ReceivedFocus(ViewportClient.Viewport);
 	}
+
+	// Static mesh editor code generally assumes the SocketManager exists, so make sure it does (in case the tab manager / sockets window hasn't already done so)
+	InitSocketManager();
 }
 
 void FStaticMeshEditor::GenerateSecondaryToolbar()
@@ -588,12 +591,7 @@ TSharedRef<SDockTab> FStaticMeshEditor::SpawnTab_SocketManager( const FSpawnTabA
 {
 	check( Args.GetTabId() == SocketManagerTabId );
 
-	if (!SocketManager)
-	{
-		FSimpleDelegate OnSocketSelectionChanged = FSimpleDelegate::CreateSP( SharedThis(this), &FStaticMeshEditor::OnSocketSelectionChanged );
-
-		SocketManager = ISocketManager::CreateSocketManager( SharedThis(this) , OnSocketSelectionChanged );
-}
+	InitSocketManager();
 
 	return SNew(SDockTab)
 		.Label( LOCTEXT("StaticMeshSocketManager_TabTitle", "Socket Manager") )
@@ -1048,6 +1046,16 @@ void FStaticMeshEditor::RequestRenameSelectedSocket()
 	SocketManager->RequestRenameSelectedSocket();
 }
 
+void FStaticMeshEditor::InitSocketManager()
+{
+	if (!SocketManager)
+	{
+		FSimpleDelegate OnSocketSelectionChanged = FSimpleDelegate::CreateSP(SharedThis(this), &FStaticMeshEditor::OnSocketSelectionChanged);
+
+		SocketManager = ISocketManager::CreateSocketManager(SharedThis(this), OnSocketSelectionChanged);
+	}
+}
+
 bool FStaticMeshEditor::IsPrimValid(const FPrimData& InPrimData) const
 {
 	if (StaticMesh->GetBodySetup())
diff --git a/Engine/Source/Editor/StaticMeshEditor/Private/StaticMeshEditor.h b/Engine/Source/Editor/StaticMeshEditor/Private/StaticMeshEditor.h
--- a/Engine/Source/Editor/StaticMeshEditor/Private/StaticMeshEditor.h
+++ b/Engine/Source/Editor/StaticMeshEditor/Private/StaticMeshEditor.h
@@ -433,6 +433,9 @@ private:
 	/** Whether we currently have any selected that can be deleted */
 	bool CanDeleteSelected() const;
 
+	/** Init SocketManager (does nothing if it already exists) */
+	void InitSocketManager();
+
 	/** Delete the currently selected sockets */
 	void DeleteSelectedSockets();
 
