Visualize Distance Field in Static Mesh Editor.

diff --git a/Engine/Shaders/Private/DistanceFieldVisualization.usf b/Engine/Shaders/Private/DistanceFieldVisualization.usf
--- a/Engine/Shaders/Private/DistanceFieldVisualization.usf
+++ b/Engine/Shaders/Private/DistanceFieldVisualization.usf
@@ -252,7 +252,7 @@ void VisualizeMeshDistanceFieldCS(
 #if !USE_GLOBAL_DISTANCE_FIELD && !MONOCHROME_VISUALIZATION
 		{
 			FDFObjectData DFObjectData = LoadDFObjectData(HitCulledObjectIndex);
-			Albedo = 0.1f * IntToColor(DFObjectData.GPUSceneInstanceIndex);
+			Albedo = 0.1f * IntToColor(DFObjectData.GPUSceneInstanceIndex + 1);
 		}
 #endif
 		Lighting *= Albedo;
diff --git a/Engine/Source/Editor/StaticMeshEditor/Private/SStaticMeshEditorViewport.cpp b/Engine/Source/Editor/StaticMeshEditor/Private/SStaticMeshEditorViewport.cpp
--- a/Engine/Source/Editor/StaticMeshEditor/Private/SStaticMeshEditorViewport.cpp
+++ b/Engine/Source/Editor/StaticMeshEditor/Private/SStaticMeshEditorViewport.cpp
@@ -230,6 +230,25 @@ bool SStaticMeshEditorViewport::IsShowNaniteFallbackVisible() const
 	return PreviewStaticMesh && PreviewStaticMesh->IsNaniteEnabled() ? true : false;
 }
 
+void SStaticMeshEditorViewport::ToggleShowDistanceField()
+{
+	if (EditorViewportClient)
+	{
+		EditorViewportClient->EngineShowFlags.SetVisualizeMeshDistanceFields(!EditorViewportClient->EngineShowFlags.VisualizeMeshDistanceFields);
+		SceneViewport->Invalidate();
+	}
+}
+
+bool SStaticMeshEditorViewport::IsShowDistanceFieldChecked() const
+{
+	return EditorViewportClient ? EditorViewportClient->EngineShowFlags.VisualizeMeshDistanceFields : false;
+}
+
+bool SStaticMeshEditorViewport::IsShowDistanceFieldVisible() const
+{
+	return true;
+}
+
 void SStaticMeshEditorViewport::UpdatePreviewSocketMeshes()
 {
 	UStaticMesh* const PreviewStaticMesh = PreviewMeshComponent ? ToRawPtr(PreviewMeshComponent->GetStaticMesh()) : nullptr;
@@ -572,6 +591,13 @@ void SStaticMeshEditorViewport::BindCommands()
 		FIsActionChecked::CreateSP(this, &SStaticMeshEditorViewport::IsShowNaniteFallbackChecked),
 		FIsActionButtonVisible::CreateSP(this, &SStaticMeshEditorViewport::IsShowNaniteFallbackVisible));
 
+	CommandList->MapAction(
+		Commands.SetShowDistanceField,
+		FExecuteAction::CreateSP(this, &SStaticMeshEditorViewport::ToggleShowDistanceField),
+		FCanExecuteAction(),
+		FIsActionChecked::CreateSP(this, &SStaticMeshEditorViewport::IsShowDistanceFieldChecked),
+		FIsActionButtonVisible::CreateSP(this, &SStaticMeshEditorViewport::IsShowDistanceFieldVisible));
+
 	CommandList->MapAction(
 		Commands.SetShowWireframe,
 		FExecuteAction::CreateSP( this, &SStaticMeshEditorViewport::SetViewModeWireframe ),
diff --git a/Engine/Source/Editor/StaticMeshEditor/Private/SStaticMeshEditorViewport.h b/Engine/Source/Editor/StaticMeshEditor/Private/SStaticMeshEditorViewport.h
--- a/Engine/Source/Editor/StaticMeshEditor/Private/SStaticMeshEditorViewport.h
+++ b/Engine/Source/Editor/StaticMeshEditor/Private/SStaticMeshEditorViewport.h
@@ -187,6 +187,12 @@ private:
 	bool IsShowNaniteFallbackChecked() const;
 
 	bool IsShowNaniteFallbackVisible() const;
+
+	void ToggleShowDistanceField();
+
+	bool IsShowDistanceFieldChecked() const;
+
+	bool IsShowDistanceFieldVisible() const;
 private:
 	
 	/** The parent tab where this viewport resides */
diff --git a/Engine/Source/Editor/StaticMeshEditor/Private/SStaticMeshEditorViewportToolBar.cpp b/Engine/Source/Editor/StaticMeshEditor/Private/SStaticMeshEditorViewportToolBar.cpp
--- a/Engine/Source/Editor/StaticMeshEditor/Private/SStaticMeshEditorViewportToolBar.cpp
+++ b/Engine/Source/Editor/StaticMeshEditor/Private/SStaticMeshEditorViewportToolBar.cpp
@@ -45,6 +45,7 @@ TSharedRef<SWidget> SStaticMeshEditorViewportToolbar::GenerateShowMenu() const
 		auto Commands = FStaticMeshEditorCommands::Get();
 
 		ShowMenuBuilder.AddMenuEntry(Commands.SetShowNaniteFallback);
+		ShowMenuBuilder.AddMenuEntry(Commands.SetShowDistanceField);
 
 		ShowMenuBuilder.BeginSection("MeshComponents", LOCTEXT("MeshComponments", "Mesh Components"));
 		ShowMenuBuilder.AddMenuEntry(Commands.SetShowSockets);
diff --git a/Engine/Source/Editor/StaticMeshEditor/Private/StaticMeshEditorActions.cpp b/Engine/Source/Editor/StaticMeshEditor/Private/StaticMeshEditorActions.cpp
--- a/Engine/Source/Editor/StaticMeshEditor/Private/StaticMeshEditorActions.cpp
+++ b/Engine/Source/Editor/StaticMeshEditor/Private/StaticMeshEditorActions.cpp
@@ -12,6 +12,7 @@
 void FStaticMeshEditorCommands::RegisterCommands()
 {
 	UI_COMMAND(SetShowNaniteFallback, "Nanite Fallback", "Toggles the display of the Nanite fallback mesh. The fallback mesh is used when Nanite is not available such as platforms that do not support it.", EUserInterfaceActionType::ToggleButton, FInputChord(EModifierKey::Control, EKeys::N));
+	UI_COMMAND(SetShowDistanceField, "Visualize Distance Field", "Toggles the display of the mesh distance field.", EUserInterfaceActionType::ToggleButton, FInputChord());
 	UI_COMMAND(SetShowWireframe, "Wireframe", "Toggles the viewmode of the Preview Pane between a lit view and a wireframe view.", EUserInterfaceActionType::ToggleButton, FInputChord());
 	UI_COMMAND(SetShowVertexColor, "Vert Colors", "Toggles vertex colors.", EUserInterfaceActionType::ToggleButton, FInputChord());
 	UI_COMMAND(SetShowPhysicalMaterialMasks, "Physical Material Masks", "Toggles physical material masks.", EUserInterfaceActionType::ToggleButton, FInputChord());
diff --git a/Engine/Source/Editor/StaticMeshEditor/Public/StaticMeshEditorActions.h b/Engine/Source/Editor/StaticMeshEditor/Public/StaticMeshEditorActions.h
--- a/Engine/Source/Editor/StaticMeshEditor/Public/StaticMeshEditorActions.h
+++ b/Engine/Source/Editor/StaticMeshEditor/Public/StaticMeshEditorActions.h
@@ -36,6 +36,7 @@ public:
 
 	 /**  */
 	TSharedPtr<FUICommandInfo> SetShowNaniteFallback;
+	TSharedPtr<FUICommandInfo> SetShowDistanceField;
 	TSharedPtr<FUICommandInfo> SetShowWireframe;
 	TSharedPtr<FUICommandInfo> SetShowVertexColor;
 	TSharedPtr<FUICommandInfo> SetShowPhysicalMaterialMasks;
