Add a new ToolMenus-based viewport toolbar to the Static Mesh editor (changed CVAR!)

* The new viewport toolbar only has a View Modes menu (Lit, Unlit, ...) at the moment. More to come!
* This change also changes the CVAR used to enable the new viewport toolbars. The new CVAR is "ToolMenusViewportToolbars" and applies to all viewport toolbars. The previous CVAR was "LevelEditor.ToolMenusViewportToolbar" and only applied to the level editor. Set the new CVAR to 0 (default) to only show the old viewport toolbars, 1 for side-by-side mode that shows both the old and new, or 2 to only show the new.

diff --git a/Engine/Source/Editor/LevelEditor/Private/SLevelViewport.cpp b/Engine/Source/Editor/LevelEditor/Private/SLevelViewport.cpp
--- a/Engine/Source/Editor/LevelEditor/Private/SLevelViewport.cpp
+++ b/Engine/Source/Editor/LevelEditor/Private/SLevelViewport.cpp
@@ -112,18 +112,6 @@ static const FName LevelEditorName("LevelEditor");
 static FAutoConsoleCommand EnableInViewportMenu(TEXT("Editor.EnableInViewportMenu"), TEXT("Enables the new in-viewport property menu"), FConsoleCommandDelegate::CreateStatic(&SLevelViewport::EnableInViewportMenu));
 bool SLevelViewport::bInViewportMenuEnabled = false;
 
-namespace UE::SLevelViewport::Private
-{
-	int32 CVarLevelEditorToolMenusViewportToolbarValue = 0;
-}
-
-static FAutoConsoleVariableRef CVarLevelEditorToolMenusViewportToolbar(
-	TEXT("LevelEditor.ToolMenusViewportToolbar"),
-	UE::SLevelViewport::Private::CVarLevelEditorToolMenusViewportToolbarValue,
-	TEXT("If set to 1, the new UToolMenus-based level editor viewport toolbar will be displayed below the current level editor viewport toolbar. If set to 2, the new will replace the current. If set to 0 (default), the current is shown but not the new."),
-	ECVF_Default
-);
-
 #define LOCTEXT_NAMESPACE "LevelViewport"
 
 // @todo Slate Hack: Disallow game UI to be used in play in viewport until GWorld problem is fixed
@@ -1906,19 +1894,21 @@ TSharedRef<FEditorViewportClient> SLevelViewport::MakeEditorViewportClient()
 
 TSharedPtr<SWidget> SLevelViewport::MakeViewportToolbar()
 {
-	const TSharedRef<SLevelViewportToolBar> OldViewportToolbar = SNew(SLevelViewportToolBar)
-		.Viewport(SharedThis(this))
-		.Visibility_Lambda([this]() -> EVisibility
-		{
-			const bool bShowOldViewportToolbar = UE::SLevelViewport::Private::CVarLevelEditorToolMenusViewportToolbarValue != 2;
-			if (!bShowOldViewportToolbar)
-			{
-				return EVisibility::Collapsed;
-			}
+	const TSharedRef<SLevelViewportToolBar> OldViewportToolbar =
+		SNew(SLevelViewportToolBar)
+			.Viewport(SharedThis(this))
+			.Visibility_Lambda(
+				[this]() -> EVisibility
+				{
+					if (!UE::UnrealEd::ShowOldViewportToolbars())
+					{
+						return EVisibility::Collapsed;
+					}
 
-			return GetToolBarVisibility();
-		})
-		.IsEnabled(FSlateApplication::Get().GetNormalExecutionAttribute());
+					return GetToolBarVisibility();
+				}
+			)
+			.IsEnabled(FSlateApplication::Get().GetNormalExecutionAttribute());
 
 	// Register the viewport toolbar if another viewport hasn't already (it's shared).
 	{
@@ -2067,20 +2057,23 @@ TSharedPtr<SWidget> SLevelViewport::MakeViewportToolbar()
 		}
 	}
 
+	// clang-format off
 	const TSharedRef<SWidget> NewViewportToolbar = SNew(SBox)
-		.Visibility_Lambda([this]() -> EVisibility
-		{
-			const bool bShowNewViewportToolbar = UE::SLevelViewport::Private::CVarLevelEditorToolMenusViewportToolbarValue > 0;
-			if (!bShowNewViewportToolbar)
+		.Visibility_Lambda(
+			[this]() -> EVisibility
 			{
-				return EVisibility::Collapsed;
-			}
+				if (!UE::UnrealEd::ShowNewViewportToolbars())
+				{
+					return EVisibility::Collapsed;
+				}
 
-			return GetToolBarVisibility();
-		})
+				return GetToolBarVisibility();
+			}
+		)
 		[
 			UToolMenus::Get()->GenerateWidget("LevelEditor.ViewportToolbar", ViewportToolbarContext)
 		];
+	// clang-format on
 
 	return 
 		SNew(SVerticalBox)
diff --git a/Engine/Source/Editor/StaticMeshEditor/Private/SStaticMeshEditorViewport.cpp b/Engine/Source/Editor/StaticMeshEditor/Private/SStaticMeshEditorViewport.cpp
--- a/Engine/Source/Editor/StaticMeshEditor/Private/SStaticMeshEditorViewport.cpp
+++ b/Engine/Source/Editor/StaticMeshEditor/Private/SStaticMeshEditorViewport.cpp
@@ -20,6 +20,8 @@
 #include "Engine/StaticMeshSocket.h"
 #include "SEditorViewportToolBarMenu.h"
 #include "Editor.h"
+#include "ToolMenus.h"
+#include "ViewportToolbar/UnrealEdViewportToolbar.h"
 #include "Widgets/Text/SRichTextBlock.h"
 
 #define HITPROXY_SOCKET	1
@@ -568,7 +570,115 @@ TSharedRef<FEditorViewportClient> SStaticMeshEditorViewport::MakeEditorViewportC
 
 TSharedPtr<SWidget> SStaticMeshEditorViewport::MakeViewportToolbar()
 {
-	return SNew(SStaticMeshEditorViewportToolbar, SharedThis(this));
+	// Register the viewport toolbar if another viewport hasn't already (it's shared).
+	const FName ViewportToolbarName = "StaticMeshEditor.ViewportToolbar";
+
+	if (!UToolMenus::Get()->IsMenuRegistered(ViewportToolbarName))
+	{
+		UToolMenu* const ViewportToolbarMenu = UToolMenus::Get()->RegisterMenu(
+			ViewportToolbarName, NAME_None /* parent */, EMultiBoxType::SlimHorizontalToolBar
+		);
+
+		// Add the left-aligned part of the viewport toolbar.
+		{
+			FToolMenuSection& LeftSection = ViewportToolbarMenu->FindOrAddSection("Left");
+		}
+
+		// Add the right-aligned part of the viewport toolbar.
+		{
+			// Add the submenus of this section as EToolMenuInsertType::Last to sort them after any
+			// default-positioned submenus external code might add.
+			FToolMenuSection& RightSection = ViewportToolbarMenu->FindOrAddSection("Right");
+			RightSection.Alignment = EToolMenuSectionAlign::Last;
+
+			// Add the "View Modes" sub menu.
+			{
+				// Stay backward-compatible with the old viewport toolbar.
+				{
+					const FName ParentSubmenuName = "UnrealEd.ViewportToolbar.View";
+					// Create our parent menu.
+					if (!UToolMenus::Get()->IsMenuRegistered(ParentSubmenuName))
+					{
+						UToolMenus::Get()->RegisterMenu(ParentSubmenuName);
+					}
+
+					// Register our ToolMenu here first, before we create the submenu, so we can set our parent.
+					UToolMenus::Get()->RegisterMenu("StaticMeshEditor.ViewportToolbar.ViewModes", ParentSubmenuName);
+				}
+
+				FToolMenuEntry ViewModesSubmenu = UE::UnrealEd::CreateViewportToolbarViewModesSubmenu();
+				ViewModesSubmenu.InsertPosition.Position = EToolMenuInsertType::Last;
+				RightSection.AddEntry(ViewModesSubmenu);
+			}
+		}
+	}
+
+	TSharedRef<SStaticMeshEditorViewportToolbar> OldViewportToolbar =
+		SNew(SStaticMeshEditorViewportToolbar, SharedThis(this))
+			.Visibility_Lambda(
+				[this]() -> EVisibility
+				{
+					return UE::UnrealEd::ShowOldViewportToolbars() ? EVisibility::Visible : EVisibility::Collapsed;
+				}
+			);
+
+	FToolMenuContext ViewportToolbarContext;
+	{
+		ViewportToolbarContext.AppendCommandList(GetCommandList());
+
+		// Add the UnrealEd viewport toolbar context.
+		{
+			UUnrealEdViewportToolbarContext* const ContextObject = NewObject<UUnrealEdViewportToolbarContext>();
+			ContextObject->Viewport = SharedThis(this);
+
+			// Hook up our toolbar's filter for supported view modes.
+			UE::UnrealEd::IsViewModeSupportedDelegate IsViewModeSupported =
+				UE::UnrealEd::IsViewModeSupportedDelegate::CreateLambda(
+					[WeakToolBar = OldViewportToolbar.ToWeakPtr()](EViewModeIndex ViewModeIndex) -> bool
+					{
+						if (TSharedPtr<SViewportToolBar> ToolBar = WeakToolBar.Pin())
+						{
+							return ToolBar->IsViewModeSupported(ViewModeIndex);
+						}
+						return true;
+					}
+				);
+
+			ViewportToolbarContext.AddObject(ContextObject);
+		}
+	}
+
+	// clang-format off
+const TSharedRef<SWidget> NewViewportToolbar = SNew(SBox)
+	.Visibility_Lambda(
+		[this]() -> EVisibility
+		{
+			return  UE::UnrealEd::ShowNewViewportToolbars() ? EVisibility::Visible: EVisibility::Collapsed;
+		}
+	)
+	[
+		UToolMenus::Get()->GenerateWidget(ViewportToolbarName, ViewportToolbarContext)
+	];
+	// clang-format on
+
+	// clang-format off
+	return SNew(SVerticalBox)
+		.Visibility( EVisibility::SelfHitTestInvisible )
+		+SVerticalBox::Slot()
+		.AutoHeight()
+		.Padding(0, 1.0f, 0, 0)
+		.VAlign(VAlign_Top)
+		[
+			OldViewportToolbar
+		]
+		+SVerticalBox::Slot()
+		.AutoHeight()
+		.Padding(0, 1.0f, 0, 0)
+		.VAlign(VAlign_Top)
+		[
+			NewViewportToolbar
+		];
+	// clang-format on
 }
 
 EVisibility SStaticMeshEditorViewport::OnGetViewportContentVisibility() const
diff --git a/Engine/Source/Editor/UnrealEd/Private/ViewportToolbar/UnrealEdViewportToolbar.cpp b/Engine/Source/Editor/UnrealEd/Private/ViewportToolbar/UnrealEdViewportToolbar.cpp
--- a/Engine/Source/Editor/UnrealEd/Private/ViewportToolbar/UnrealEdViewportToolbar.cpp
+++ b/Engine/Source/Editor/UnrealEd/Private/ViewportToolbar/UnrealEdViewportToolbar.cpp
@@ -18,9 +18,33 @@
 
 #define LOCTEXT_NAMESPACE "UnrealEdViewportToolbar"
 
+namespace UE::UnrealEd::Private
+{
+int32 CVarToolMenusViewportToolbarsValue = 0;
+}
+
+static FAutoConsoleVariableRef CVarToolMenusViewportToolbars(
+	TEXT("ToolMenusViewportToolbars"),
+	UE::UnrealEd::Private::CVarToolMenusViewportToolbarsValue,
+	TEXT("Control whether the new ToolMenus-based viewport toolbars are enabled across the editor. Set to 0 (default) "
+		 "to show only the old viewport toolbars. Set to 1 for side-by-side mode where both the old and new viewport "
+		 "toolbars are shown. Set to 2 to show only the new viewport toolbars."),
+	ECVF_Default
+);
+
 namespace UE::UnrealEd
 {
 
+bool ShowOldViewportToolbars()
+{
+	return Private::CVarToolMenusViewportToolbarsValue <= 1;
+}
+
+bool ShowNewViewportToolbars()
+{
+	return Private::CVarToolMenusViewportToolbarsValue >= 1;
+}
+
 FText GetViewModesSubmenuLabel(TWeakPtr<SEditorViewport> InViewport)
 {
 	FText Label = LOCTEXT("ViewMenuTitle_Default", "View");
diff --git a/Engine/Source/Editor/UnrealEd/Public/ViewportToolbar/UnrealEdViewportToolbar.h b/Engine/Source/Editor/UnrealEd/Public/ViewportToolbar/UnrealEdViewportToolbar.h
--- a/Engine/Source/Editor/UnrealEd/Public/ViewportToolbar/UnrealEdViewportToolbar.h
+++ b/Engine/Source/Editor/UnrealEd/Public/ViewportToolbar/UnrealEdViewportToolbar.h
@@ -15,6 +15,12 @@ struct FToolMenuEntry;
 namespace UE::UnrealEd
 {
 
+/** The value of this function is controlled by the CVAR "ToolMenusViewportToolbars". */
+UNREALED_API bool ShowOldViewportToolbars();
+
+/** The value of this function is controlled by the CVAR "ToolMenusViewportToolbars". */
+UNREALED_API bool ShowNewViewportToolbars();
+
 UNREALED_API FText GetViewModesSubmenuLabel(TWeakPtr<SEditorViewport> InViewport);
 
 /**
