[Viewport Toolbar] Use UE::UnrealEd Transform, Selection and Snapping menus from Blueprint and Static Mesh editors

diff --git a/Engine/Source/Editor/Kismet/Private/SSCSEditorViewport.cpp b/Engine/Source/Editor/Kismet/Private/SSCSEditorViewport.cpp
--- a/Engine/Source/Editor/Kismet/Private/SSCSEditorViewport.cpp
+++ b/Engine/Source/Editor/Kismet/Private/SSCSEditorViewport.cpp
@@ -339,13 +339,12 @@ TSharedPtr<SWidget> SSCSEditorViewport::MakeViewportToolbar()
 				LeftSection.AddEntry(SelectionSubmenu);
 			}
 
-			// TODO: Add this.
-			// // Add the "Snapping" sub menu.
-			// {
-			// 	FToolMenuEntry SnappingSubmenu = UE::LevelEditor::CreateViewportToolbarSnappingSubmenu();
-			// 	SnappingSubmenu.InsertPosition.Position = EToolMenuInsertType::First;
-			// 	LeftSection.AddEntry(SnappingSubmenu);
-			// }
+			// Add the "Snapping" sub menu.
+			{
+				FToolMenuEntry SnappingSubmenu = UE::UnrealEd::CreateViewportToolbarSnappingSubmenu();
+				SnappingSubmenu.InsertPosition.Position = EToolMenuInsertType::First;
+				LeftSection.AddEntry(SnappingSubmenu);
+			}
 		}
 
 		// Add the right-aligned part of the viewport toolbar.
diff --git a/Engine/Source/Editor/StaticMeshEditor/Private/SStaticMeshEditorViewport.cpp b/Engine/Source/Editor/StaticMeshEditor/Private/SStaticMeshEditorViewport.cpp
--- a/Engine/Source/Editor/StaticMeshEditor/Private/SStaticMeshEditorViewport.cpp
+++ b/Engine/Source/Editor/StaticMeshEditor/Private/SStaticMeshEditorViewport.cpp
@@ -82,7 +82,7 @@ void SStaticMeshEditorViewport::PopulateViewportOverlays(TSharedRef<SOverlay> Ov
 	Overlay->AddSlot()
 		.VAlign(VAlign_Top)
 		.HAlign(HAlign_Left)
-		.Padding(FMargin(6.0f, 36.0f, 6.0f, 6.0f))
+		.Padding(TAttribute<FMargin>(this, &SStaticMeshEditorViewport::GetOverlayMargin))
 		[
 			SNew(SBorder)
 			.BorderImage( FAppStyle::Get().GetBrush( "FloatingBorder" ) )
@@ -251,6 +251,13 @@ bool SStaticMeshEditorViewport::IsShowDistanceFieldVisible() const
 	return true;
 }
 
+FMargin SStaticMeshEditorViewport::GetOverlayMargin() const
+{
+	return UE::UnrealEd::ShowNewViewportToolbars() && UE::UnrealEd::ShowOldViewportToolbars()
+			 ? FMargin(6.0f, 72.0f, 6.0f, 6.0f)
+			 : FMargin(6.0f, 36.0f, 6.0f, 6.0f);
+}
+
 void SStaticMeshEditorViewport::UpdatePreviewSocketMeshes()
 {
 	UStaticMesh* const PreviewStaticMesh = PreviewMeshComponent ? ToRawPtr(PreviewMeshComponent->GetStaticMesh()) : nullptr;
@@ -582,6 +589,27 @@ TSharedPtr<SWidget> SStaticMeshEditorViewport::MakeViewportToolbar()
 		// Add the left-aligned part of the viewport toolbar.
 		{
 			FToolMenuSection& LeftSection = ViewportToolbarMenu->FindOrAddSection("Left");
+
+			// Add the "Transforms" sub menu.
+			{
+				FToolMenuEntry TransformsSubmenu = UE::UnrealEd::CreateViewportToolbarTransformsSection();
+				TransformsSubmenu.InsertPosition.Position = EToolMenuInsertType::First;
+				LeftSection.AddEntry(TransformsSubmenu);
+			}
+
+			// Add the "Selection" sub menu.
+			{
+				FToolMenuEntry SelectionSubmenu = UE::UnrealEd::CreateViewportToolbarSelectionSection();
+				SelectionSubmenu.InsertPosition.Position = EToolMenuInsertType::First;
+				LeftSection.AddEntry(SelectionSubmenu);
+			}
+
+			// Add the "Snapping" sub menu.
+			{
+				FToolMenuEntry SnappingSubmenu = UE::UnrealEd::CreateViewportToolbarSnappingSubmenu();
+				SnappingSubmenu.InsertPosition.Position = EToolMenuInsertType::First;
+				LeftSection.AddEntry(SnappingSubmenu);
+			}
 		}
 
 		// Add the right-aligned part of the viewport toolbar.
diff --git a/Engine/Source/Editor/StaticMeshEditor/Private/SStaticMeshEditorViewport.h b/Engine/Source/Editor/StaticMeshEditor/Private/SStaticMeshEditorViewport.h
--- a/Engine/Source/Editor/StaticMeshEditor/Private/SStaticMeshEditorViewport.h
+++ b/Engine/Source/Editor/StaticMeshEditor/Private/SStaticMeshEditorViewport.h
@@ -193,6 +193,9 @@ private:
 	bool IsShowDistanceFieldChecked() const;
 
 	bool IsShowDistanceFieldVisible() const;
+
+	/** Returns the appropriate margin for the overlay, based on the current viewport toolbar configuration */
+	FMargin GetOverlayMargin() const;
 private:
 	
 	/** The parent tab where this viewport resides */
