[Viewport Toolbar] Add new "Performance & Scalability" submenu to the new Static Mesh viewpor toolbar

diff --git a/Engine/Source/Editor/StaticMeshEditor/Private/SStaticMeshEditorViewport.cpp b/Engine/Source/Editor/StaticMeshEditor/Private/SStaticMeshEditorViewport.cpp
--- a/Engine/Source/Editor/StaticMeshEditor/Private/SStaticMeshEditorViewport.cpp
+++ b/Engine/Source/Editor/StaticMeshEditor/Private/SStaticMeshEditorViewport.cpp
@@ -637,6 +637,10 @@ TSharedPtr<SWidget> SStaticMeshEditorViewport::MakeViewportToolbar()
 				FToolMenuEntry ViewModesSubmenu = UE::UnrealEd::CreateViewportToolbarViewModesSubmenu();
 				ViewModesSubmenu.InsertPosition.Position = EToolMenuInsertType::Last;
 				RightSection.AddEntry(ViewModesSubmenu);
+
+				FToolMenuEntry PerformanceAndScalabilitySubmenu = UE::UnrealEd::CreatePerformanceAndScalabilitySubmenu(SharedThis(this));
+				PerformanceAndScalabilitySubmenu.InsertPosition.Position = EToolMenuInsertType::Last;
+				RightSection.AddEntry(PerformanceAndScalabilitySubmenu);
 			}
 		}
 	}
@@ -677,7 +681,7 @@ TSharedPtr<SWidget> SStaticMeshEditorViewport::MakeViewportToolbar()
 	}
 
 	// clang-format off
-const TSharedRef<SWidget> NewViewportToolbar = SNew(SBox)
+	const TSharedRef<SWidget> NewViewportToolbar = SNew(SBox)
 	.Visibility_Lambda(
 		[this]() -> EVisibility
 		{
