[Viewport Toolbar] Move the new viewport toolbars above their viewports instead of on top to not obscure pixels

Updates the following editors:

- Level
- Blueprint
- Static Mesh
- Material/Instance/Function

diff --git a/Engine/Source/Editor/Kismet/Private/SSCSEditorViewport.cpp b/Engine/Source/Editor/Kismet/Private/SSCSEditorViewport.cpp
--- a/Engine/Source/Editor/Kismet/Private/SSCSEditorViewport.cpp
+++ b/Engine/Source/Editor/Kismet/Private/SSCSEditorViewport.cpp
@@ -315,6 +315,22 @@ TSharedPtr<SWidget> SSCSEditorViewport::MakeViewportToolbar()
 				}
 			);
 
+	// clang-format off
+	return 
+		SNew(SVerticalBox)
+		.Visibility( EVisibility::SelfHitTestInvisible )
+		+SVerticalBox::Slot()
+		.AutoHeight()
+		.Padding(0, 1.0f, 0, 0)
+		.VAlign(VAlign_Top)
+		[
+			OldViewportToolbar
+		];
+	// clang-format on
+}
+
+TSharedPtr<SWidget> SSCSEditorViewport::BuildViewportToolbar()
+{
 	// Register the viewport toolbar if another viewport hasn't already (it's shared).
 	const FName ViewportToolbarMenuName = "SCSEditor.ViewportToolbar";
 	if (!UToolMenus::Get()->IsMenuRegistered(ViewportToolbarMenuName))
@@ -431,18 +447,25 @@ TSharedPtr<SWidget> SSCSEditorViewport::MakeViewportToolbar()
 			UUnrealEdViewportToolbarContext* const ContextObject = NewObject<UUnrealEdViewportToolbarContext>();
 			ContextObject->Viewport = SharedThis(this);
 
-			// Hook up our toolbar's filter for supported view modes.
-			UE::UnrealEd::IsViewModeSupportedDelegate IsViewModeSupported =
-				UE::UnrealEd::IsViewModeSupportedDelegate::CreateLambda(
-					[WeakToolBar = OldViewportToolbar.ToWeakPtr()](EViewModeIndex ViewModeIndex) -> bool
+			// Setup the callback to filter available view modes
+			ContextObject->IsViewModeSupported = UE::UnrealEd::IsViewModeSupportedDelegate::CreateLambda(
+				[](EViewModeIndex ViewModeIndex) -> bool
+				{
+					// This code is taken from SViewportToolBar::IsViewModeSupported
+					// SSCSEditorViewportToolBar does not override it, so we just take it as-is
+					// TODO: maybe create a private function for it, or move IsViewModeSupported to SEditorViewport
+
+					switch (ViewModeIndex)
 					{
-						if (TSharedPtr<SViewportToolBar> ToolBar = WeakToolBar.Pin())
-						{
-							return ToolBar->IsViewModeSupported(ViewModeIndex);
-						}
+					case VMI_PrimitiveDistanceAccuracy:
+					case VMI_MaterialTextureScaleAccuracy:
+					case VMI_RequiredTextureResolution:
+						return false;
+					default:
 						return true;
 					}
-				);
+				}
+			);
 
 			ViewportToolbarContext.AddObject(ContextObject);
 		}
@@ -462,25 +485,7 @@ TSharedPtr<SWidget> SSCSEditorViewport::MakeViewportToolbar()
 		);
 	// clang-format on
 
-	// clang-format off
-	return 
-		SNew(SVerticalBox)
-		.Visibility( EVisibility::SelfHitTestInvisible )
-		+SVerticalBox::Slot()
-		.AutoHeight()
-		.Padding(0, 1.0f, 0, 0)
-		.VAlign(VAlign_Top)
-		[
-			OldViewportToolbar
-		]
-		+SVerticalBox::Slot()
-		.AutoHeight()
-		.Padding(0, 1.0f, 0, 0)
-		.VAlign(VAlign_Top)
-		[
-			NewViewportToolbar
-		];
-	// clang-format on
+	return NewViewportToolbar;
 }
 
 void SSCSEditorViewport::PopulateViewportOverlays(TSharedRef<class SOverlay> Overlay)
diff --git a/Engine/Source/Editor/Kismet/Private/SSCSEditorViewport.h b/Engine/Source/Editor/Kismet/Private/SSCSEditorViewport.h
--- a/Engine/Source/Editor/Kismet/Private/SSCSEditorViewport.h
+++ b/Engine/Source/Editor/Kismet/Private/SSCSEditorViewport.h
@@ -97,6 +97,7 @@ protected:
 	/** SEditorViewport interface */
 	virtual TSharedRef<class FEditorViewportClient> MakeEditorViewportClient() override;
 	virtual TSharedPtr<class SWidget> MakeViewportToolbar() override;
+	virtual TSharedPtr<SWidget> BuildViewportToolbar() override;
 	virtual void PopulateViewportOverlays(TSharedRef<class SOverlay> Overlay) override;
 	virtual void BindCommands() override;
 
diff --git a/Engine/Source/Editor/LevelEditor/Private/SLevelViewport.cpp b/Engine/Source/Editor/LevelEditor/Private/SLevelViewport.cpp
--- a/Engine/Source/Editor/LevelEditor/Private/SLevelViewport.cpp
+++ b/Engine/Source/Editor/LevelEditor/Private/SLevelViewport.cpp
@@ -1910,6 +1910,30 @@ TSharedPtr<SWidget> SLevelViewport::MakeViewportToolbar()
 			)
 			.IsEnabled(FSlateApplication::Get().GetNormalExecutionAttribute());
 
+	return
+		// clang-format off
+		SNew(SVerticalBox)
+		.Visibility( EVisibility::SelfHitTestInvisible )
+		+SVerticalBox::Slot()
+		.AutoHeight()
+		.Padding(0, 1.0f, 0, 0)
+		.VAlign(VAlign_Top)
+		[
+			OldViewportToolbar
+		]
+		+SVerticalBox::Slot()
+		.VAlign(VAlign_Top)
+		.HAlign(HAlign_Left)
+		[
+			SNew(SActorPilotViewportToolbar)
+			.Viewport( SharedThis( this ) )
+			.Visibility(this, &SLevelViewport::GetLockedIconVisibility)
+		];
+	// clang-format on
+}
+
+TSharedPtr<SWidget> SLevelViewport::BuildViewportToolbar()
+{
 	// Register the viewport toolbar if another viewport hasn't already (it's shared).
 	{
 		const FName LevelEditorViewportToolbarName = "LevelEditor.ViewportToolbar";
@@ -2038,7 +2062,7 @@ TSharedPtr<SWidget> SLevelViewport::MakeViewportToolbar()
 			ViewportToolbarContext.AddExtender(UE::LevelEditor::GetViewModesLegacyExtenders());
 
 			FLevelEditorModule& LevelEditorModule =
-					FModuleManager::GetModuleChecked<FLevelEditorModule>(TEXT("LevelEditor"));
+				FModuleManager::GetModuleChecked<FLevelEditorModule>(TEXT("LevelEditor"));
 
 			TSharedRef<FUICommandList> CommandListRef = GetCommandList().ToSharedRef();
 
@@ -2088,47 +2112,23 @@ TSharedPtr<SWidget> SLevelViewport::MakeViewportToolbar()
 
 	// clang-format off
 	const TSharedRef<SWidget> NewViewportToolbar = SNew(SBox)
-		.Visibility_Lambda(
-			[this]() -> EVisibility
+	.Visibility_Lambda(
+		[this]() -> EVisibility
+		{
+			if (!UE::UnrealEd::ShowNewViewportToolbars())
 			{
-				if (!UE::UnrealEd::ShowNewViewportToolbars())
-				{
-					return EVisibility::Collapsed;
-				}
-
-				return GetToolBarVisibility();
+				return EVisibility::Collapsed;
 			}
-		)
-		[
-			UToolMenus::Get()->GenerateWidget("LevelEditor.ViewportToolbar", ViewportToolbarContext)
-		];
+
+			return GetToolBarVisibility();
+		}
+	)
+	[
+		UToolMenus::Get()->GenerateWidget("LevelEditor.ViewportToolbar", ViewportToolbarContext)
+	];
 	// clang-format on
 
-	return 
-		SNew(SVerticalBox)
-		.Visibility( EVisibility::SelfHitTestInvisible )
-		+SVerticalBox::Slot()
-		.AutoHeight()
-		.Padding(0, 1.0f, 0, 0)
-		.VAlign(VAlign_Top)
-		[
-			OldViewportToolbar
-		]
-		+SVerticalBox::Slot()
-		.AutoHeight()
-		.Padding(0, 1.0f, 0, 0)
-		.VAlign(VAlign_Top)
-		[
-			NewViewportToolbar
-		]
-		+SVerticalBox::Slot()
-		.VAlign(VAlign_Top)
-		.HAlign(HAlign_Left)
-		[
-			SNew(SActorPilotViewportToolbar)
-			.Viewport( SharedThis( this ) )
-			.Visibility(this, &SLevelViewport::GetLockedIconVisibility)
-		];
+	return NewViewportToolbar;
 }
 
 void SLevelViewport::OnUndo()
diff --git a/Engine/Source/Editor/LevelEditor/Public/SLevelViewport.h b/Engine/Source/Editor/LevelEditor/Public/SLevelViewport.h
--- a/Engine/Source/Editor/LevelEditor/Public/SLevelViewport.h
+++ b/Engine/Source/Editor/LevelEditor/Public/SLevelViewport.h
@@ -451,6 +451,7 @@ protected:
 	/** SEditorViewport interface */
 	virtual TSharedRef<FEditorViewportClient> MakeEditorViewportClient() override;
 	virtual TSharedPtr<SWidget> MakeViewportToolbar() override;
+	virtual TSharedPtr<SWidget> BuildViewportToolbar() override;
 
 	virtual void OnIncrementPositionGridSize() override;
 	virtual void OnDecrementPositionGridSize() override;
diff --git a/Engine/Source/Editor/MaterialEditor/Private/SMaterialEditorViewport.cpp b/Engine/Source/Editor/MaterialEditor/Private/SMaterialEditorViewport.cpp
--- a/Engine/Source/Editor/MaterialEditor/Private/SMaterialEditorViewport.cpp
+++ b/Engine/Source/Editor/MaterialEditor/Private/SMaterialEditorViewport.cpp
@@ -938,22 +938,8 @@ TSharedRef<FEditorViewportClient> SMaterialEditor3DPreviewViewport::MakeEditorVi
 	return EditorViewportClient.ToSharedRef();
 }
 
-void SMaterialEditor3DPreviewViewport::PopulateViewportOverlays(TSharedRef<class SOverlay> Overlay)
+TSharedPtr<SWidget> SMaterialEditor3DPreviewViewport::BuildViewportToolbar()
 {
-	const TSharedRef<SMaterialEditorViewportToolBar> OldViewportToolbar =
-		// clang-format off
-		SNew(SMaterialEditorViewportToolBar, SharedThis(this))
-			.Visibility_Lambda([this]() -> EVisibility
-			{
-				if (!UE::UnrealEd::ShowOldViewportToolbars())
-				{
-					return EVisibility::Collapsed;
-				}
-
-				return EVisibility::Visible;
-			});
-	// clang-format on
-
 	const FName MaterialEditorViewportToolbarName = "MaterialEditor.ViewportToolbar";
 
 	// Register the viewport toolbar if another viewport hasn't already (it's shared).
@@ -1054,6 +1040,25 @@ void SMaterialEditor3DPreviewViewport::PopulateViewportOverlays(TSharedRef<class
 		];
 		// clang-format on
 
+	return NewViewportToolbar;
+}
+
+void SMaterialEditor3DPreviewViewport::PopulateViewportOverlays(TSharedRef<class SOverlay> Overlay)
+{
+	const TSharedRef<SMaterialEditorViewportToolBar> OldViewportToolbar =
+		// clang-format off
+		SNew(SMaterialEditorViewportToolBar, SharedThis(this))
+			.Visibility_Lambda([this]() -> EVisibility
+			{
+				if (!UE::UnrealEd::ShowOldViewportToolbars())
+				{
+					return EVisibility::Collapsed;
+				}
+
+				return EVisibility::Visible;
+			});
+	// clang-format on
+
 	Overlay->AddSlot()
 		// clang-format off
 		.VAlign(VAlign_Top)
@@ -1067,13 +1072,6 @@ void SMaterialEditor3DPreviewViewport::PopulateViewportOverlays(TSharedRef<class
 			[
 				OldViewportToolbar
 			]
-			+ SVerticalBox::Slot()
-			.AutoHeight()
-			.Padding(0.0f, 1.0f, 0.0f, 0.0f)
-			.VAlign(VAlign_Top)
-			[
-				NewViewportToolbar
-			]
 		];
 		// clang-format on
 
diff --git a/Engine/Source/Editor/MaterialEditor/Private/SMaterialEditorViewport.h b/Engine/Source/Editor/MaterialEditor/Private/SMaterialEditorViewport.h
--- a/Engine/Source/Editor/MaterialEditor/Private/SMaterialEditorViewport.h
+++ b/Engine/Source/Editor/MaterialEditor/Private/SMaterialEditorViewport.h
@@ -95,6 +95,7 @@ public:
 protected:
 	/** SEditorViewport interface */
 	virtual TSharedRef<FEditorViewportClient> MakeEditorViewportClient() override;
+	virtual TSharedPtr<SWidget> BuildViewportToolbar() override;
 	virtual void PopulateViewportOverlays(TSharedRef<class SOverlay> Overlay) override;
 	virtual EVisibility OnGetViewportContentVisibility() const override;
 	virtual void BindCommands() override;
diff --git a/Engine/Source/Editor/StaticMeshEditor/Private/SStaticMeshEditorViewport.cpp b/Engine/Source/Editor/StaticMeshEditor/Private/SStaticMeshEditorViewport.cpp
--- a/Engine/Source/Editor/StaticMeshEditor/Private/SStaticMeshEditorViewport.cpp
+++ b/Engine/Source/Editor/StaticMeshEditor/Private/SStaticMeshEditorViewport.cpp
@@ -577,74 +577,6 @@
 
 TSharedPtr<SWidget> SStaticMeshEditorViewport::MakeViewportToolbar()
 {
-	// Register the viewport toolbar if another viewport hasn't already (it's shared).
-	const FName ViewportToolbarName = "StaticMeshEditor.ViewportToolbar";
-
-	if (!UToolMenus::Get()->IsMenuRegistered(ViewportToolbarName))
-	{
-		UToolMenu* const ViewportToolbarMenu = UToolMenus::Get()->RegisterMenu(
-			ViewportToolbarName, NAME_None /* parent */, EMultiBoxType::SlimHorizontalToolBar
-		);
-
-		// Add the left-aligned part of the viewport toolbar.
-		{
-			FToolMenuSection& LeftSection = ViewportToolbarMenu->FindOrAddSection("Left");
-
-			// Add the "Transforms" sub menu.
-			{
-				FToolMenuEntry TransformsSubmenu = UE::UnrealEd::CreateViewportToolbarTransformsSection();
-				TransformsSubmenu.InsertPosition.Position = EToolMenuInsertType::First;
-				LeftSection.AddEntry(TransformsSubmenu);
-			}
-
-			// Add the "Selection" sub menu.
-			{
-				FToolMenuEntry SelectionSubmenu = UE::UnrealEd::CreateViewportToolbarSelectionSection();
-				SelectionSubmenu.InsertPosition.Position = EToolMenuInsertType::First;
-				LeftSection.AddEntry(SelectionSubmenu);
-			}
-
-			// Add the "Snapping" sub menu.
-			{
-				FToolMenuEntry SnappingSubmenu = UE::UnrealEd::CreateViewportToolbarSnappingSubmenu();
-				SnappingSubmenu.InsertPosition.Position = EToolMenuInsertType::First;
-				LeftSection.AddEntry(SnappingSubmenu);
-			}
-		}
-
-		// Add the right-aligned part of the viewport toolbar.
-		{
-			// Add the submenus of this section as EToolMenuInsertType::Last to sort them after any
-			// default-positioned submenus external code might add.
-			FToolMenuSection& RightSection = ViewportToolbarMenu->FindOrAddSection("Right");
-			RightSection.Alignment = EToolMenuSectionAlign::Last;
-
-			// Add the "View Modes" sub menu.
-			{
-				// Stay backward-compatible with the old viewport toolbar.
-				{
-					const FName ParentSubmenuName = "UnrealEd.ViewportToolbar.View";
-					// Create our parent menu.
-					if (!UToolMenus::Get()->IsMenuRegistered(ParentSubmenuName))
-					{
-						UToolMenus::Get()->RegisterMenu(ParentSubmenuName);
-					}
-
-					// Register our ToolMenu here first, before we create the submenu, so we can set our parent.
-					UToolMenus::Get()->RegisterMenu("StaticMeshEditor.ViewportToolbar.ViewModes", ParentSubmenuName);
-				}
-
-				FToolMenuEntry ViewModesSubmenu = UE::UnrealEd::CreateViewportToolbarViewModesSubmenu();
-				ViewModesSubmenu.InsertPosition.Position = EToolMenuInsertType::Last;
-				RightSection.AddEntry(ViewModesSubmenu);
-
-				FToolMenuEntry PerformanceAndScalabilitySubmenu = UE::UnrealEd::CreatePerformanceAndScalabilitySubmenu();
-				PerformanceAndScalabilitySubmenu.InsertPosition.Position = EToolMenuInsertType::Last;
-				RightSection.AddEntry(PerformanceAndScalabilitySubmenu);
-			}
-		}
-	}
-
 	TSharedRef<SStaticMeshEditorViewportToolbar> OldViewportToolbar =
 		SNew(SStaticMeshEditorViewportToolbar, SharedThis(this))
 			.Visibility_Lambda(
@@ -654,45 +586,6 @@
 				}
 			);
 
-	FToolMenuContext ViewportToolbarContext;
-	{
-		ViewportToolbarContext.AppendCommandList(GetCommandList());
-
-		// Add the UnrealEd viewport toolbar context.
-		{
-			UUnrealEdViewportToolbarContext* const ContextObject = NewObject<UUnrealEdViewportToolbarContext>();
-			ContextObject->Viewport = SharedThis(this);
-
-			// Hook up our toolbar's filter for supported view modes.
-			UE::UnrealEd::IsViewModeSupportedDelegate IsViewModeSupported =
-				UE::UnrealEd::IsViewModeSupportedDelegate::CreateLambda(
-					[WeakToolBar = OldViewportToolbar.ToWeakPtr()](EViewModeIndex ViewModeIndex) -> bool
-					{
-						if (TSharedPtr<SViewportToolBar> ToolBar = WeakToolBar.Pin())
-						{
-							return ToolBar->IsViewModeSupported(ViewModeIndex);
-						}
-						return true;
-					}
-				);
-
-			ViewportToolbarContext.AddObject(ContextObject);
-		}
-	}
-
 	// clang-format off
-	const TSharedRef<SWidget> NewViewportToolbar = SNew(SBox)
-	.Visibility_Lambda(
-		[this]() -> EVisibility
-		{
-			return  UE::UnrealEd::ShowNewViewportToolbars() ? EVisibility::Visible: EVisibility::Collapsed;
-		}
-	)
-	[
-		UToolMenus::Get()->GenerateWidget(ViewportToolbarName, ViewportToolbarContext)
-	];
-	// clang-format on
-
-	// clang-format off
 	return SNew(SVerticalBox)
 		.Visibility( EVisibility::SelfHitTestInvisible )
@@ -702,13 +595,6 @@
 		.VAlign(VAlign_Top)
 		[
 			OldViewportToolbar
-		]
-		+SVerticalBox::Slot()
-		.AutoHeight()
-		.Padding(0, 1.0f, 0, 0)
-		.VAlign(VAlign_Top)
-		[
-			NewViewportToolbar
 		];
 	// clang-format on
 }
@@ -858,21 +744,21 @@
 {
 	// If we have selected sockets, focus on them
 	UStaticMeshSocket* SelectedSocket = StaticMeshEditorPtr.Pin()->GetSelectedSocket();
-	if( SelectedSocket && PreviewMeshComponent )
+	if (SelectedSocket && PreviewMeshComponent)
 	{
 		FTransform SocketTransform;
-		SelectedSocket->GetSocketTransform( SocketTransform, PreviewMeshComponent );
+		SelectedSocket->GetSocketTransform(SocketTransform, PreviewMeshComponent);
 
 		const FVector Extent(30.0f);
 
 		const FVector Origin = SocketTransform.GetLocation();
 		const FBox Box(Origin - Extent, Origin + Extent);
 
-		EditorViewportClient->FocusViewportOnBox( Box );
+		EditorViewportClient->FocusViewportOnBox(Box);
 		return;
 	}
 
-	// If we have selected primitives, focus on them 
+	// If we have selected primitives, focus on them
 	FBox Box(ForceInit);
 	const bool bSelectedPrim = StaticMeshEditorPtr.Pin()->CalcSelectedPrimsAABB(Box);
 	if (bSelectedPrim)
@@ -882,11 +768,130 @@
 	}
 
 	// Fallback to focusing on the mesh, if nothing else
-	if( PreviewMeshComponent )
+	if (PreviewMeshComponent)
 	{
-		EditorViewportClient->FocusViewportOnBox( PreviewMeshComponent->Bounds.GetBox() );
+		EditorViewportClient->FocusViewportOnBox(PreviewMeshComponent->Bounds.GetBox());
 		return;
 	}
 }
 
+TSharedPtr<SWidget> SStaticMeshEditorViewport::BuildViewportToolbar()
+{
+	// Register the viewport toolbar if another viewport hasn't already (it's shared).
+	const FName ViewportToolbarName = "StaticMeshEditor.ViewportToolbar";
+
+	if (!UToolMenus::Get()->IsMenuRegistered(ViewportToolbarName))
+	{
+		UToolMenu* const ViewportToolbarMenu = UToolMenus::Get()->RegisterMenu(
+			ViewportToolbarName, NAME_None /* parent */, EMultiBoxType::SlimHorizontalToolBar
+		);
+
+		// Add the left-aligned part of the viewport toolbar.
+		{
+			FToolMenuSection& LeftSection = ViewportToolbarMenu->FindOrAddSection("Left");
+
+			// Add the "Transforms" sub menu.
+			{
+				FToolMenuEntry TransformsSubmenu = UE::UnrealEd::CreateViewportToolbarTransformsSection();
+				TransformsSubmenu.InsertPosition.Position = EToolMenuInsertType::First;
+				LeftSection.AddEntry(TransformsSubmenu);
+			}
+
+			// Add the "Selection" sub menu.
+			{
+				FToolMenuEntry SelectionSubmenu = UE::UnrealEd::CreateViewportToolbarSelectionSection();
+				SelectionSubmenu.InsertPosition.Position = EToolMenuInsertType::First;
+				LeftSection.AddEntry(SelectionSubmenu);
+			}
+
+			// Add the "Snapping" sub menu.
+			{
+				FToolMenuEntry SnappingSubmenu = UE::UnrealEd::CreateViewportToolbarSnappingSubmenu();
+				SnappingSubmenu.InsertPosition.Position = EToolMenuInsertType::First;
+				LeftSection.AddEntry(SnappingSubmenu);
+			}
+		}
+
+		// Add the right-aligned part of the viewport toolbar.
+		{
+			// Add the submenus of this section as EToolMenuInsertType::Last to sort them after any
+			// default-positioned submenus external code might add.
+			FToolMenuSection& RightSection = ViewportToolbarMenu->FindOrAddSection("Right");
+			RightSection.Alignment = EToolMenuSectionAlign::Last;
+
+			// Add the "View Modes" sub menu.
+			{
+				// Stay backward-compatible with the old viewport toolbar.
+				{
+					const FName ParentSubmenuName = "UnrealEd.ViewportToolbar.View";
+					// Create our parent menu.
+					if (!UToolMenus::Get()->IsMenuRegistered(ParentSubmenuName))
+					{
+						UToolMenus::Get()->RegisterMenu(ParentSubmenuName);
+					}
+
+					// Register our ToolMenu here first, before we create the submenu, so we can set our parent.
+					UToolMenus::Get()->RegisterMenu("StaticMeshEditor.ViewportToolbar.ViewModes", ParentSubmenuName);
+				}
+
+				FToolMenuEntry ViewModesSubmenu = UE::UnrealEd::CreateViewportToolbarViewModesSubmenu();
+				ViewModesSubmenu.InsertPosition.Position = EToolMenuInsertType::Last;
+				RightSection.AddEntry(ViewModesSubmenu);
+
+				FToolMenuEntry PerformanceAndScalabilitySubmenu = UE::UnrealEd::CreatePerformanceAndScalabilitySubmenu();
+				PerformanceAndScalabilitySubmenu.InsertPosition.Position = EToolMenuInsertType::Last;
+				RightSection.AddEntry(PerformanceAndScalabilitySubmenu);
+			}
+		}
+	}
+
+	FToolMenuContext ViewportToolbarContext;
+	{
+		ViewportToolbarContext.AppendCommandList(GetCommandList());
+
+		// Add the UnrealEd viewport toolbar context.
+		{
+			UUnrealEdViewportToolbarContext* const ContextObject = NewObject<UUnrealEdViewportToolbarContext>();
+			ContextObject->Viewport = SharedThis(this);
+
+			// Hook up our toolbar's filter for supported view modes.
+			ContextObject->IsViewModeSupported = UE::UnrealEd::IsViewModeSupportedDelegate::CreateLambda(
+				[](EViewModeIndex ViewModeIndex) -> bool
+				{
+					// This code is taken from SViewportToolBar::IsViewModeSupported
+					// SSCSEditorViewportToolBar does not override it, so we just take it as-is
+					// TODO: maybe create a private function for it, or move IsViewModeSupported to SEditorViewport
+
+					switch (ViewModeIndex)
+					{
+					case VMI_PrimitiveDistanceAccuracy:
+					case VMI_MaterialTextureScaleAccuracy:
+					case VMI_RequiredTextureResolution:
+						return false;
+					default:
+						return true;
+					}
+				}
+			);
+
+			ViewportToolbarContext.AddObject(ContextObject);
+		}
+	}
+
+	// clang-format off
+	const TSharedRef<SWidget> NewViewportToolbar = SNew(SBox)
+	.Visibility_Lambda(
+		[this]() -> EVisibility
+		{
+			return  UE::UnrealEd::ShowNewViewportToolbars() ? EVisibility::Visible: EVisibility::Collapsed;
+		}
+	)
+	[
+		UToolMenus::Get()->GenerateWidget(ViewportToolbarName, ViewportToolbarContext)
+	];
+	// clang-format on
+
+	return NewViewportToolbar;
+}
+
 #undef LOCTEXT_NAMESPACE
diff --git a/Engine/Source/Editor/StaticMeshEditor/Private/SStaticMeshEditorViewport.h b/Engine/Source/Editor/StaticMeshEditor/Private/SStaticMeshEditorViewport.h
--- a/Engine/Source/Editor/StaticMeshEditor/Private/SStaticMeshEditorViewport.h
+++ b/Engine/Source/Editor/StaticMeshEditor/Private/SStaticMeshEditorViewport.h
@@ -136,6 +136,7 @@ protected:
 	virtual EVisibility OnGetViewportContentVisibility() const override;
 	virtual void BindCommands() override;
 	virtual void OnFocusViewportToSelection() override;
+	virtual TSharedPtr<SWidget> BuildViewportToolbar() override;
 	virtual TSharedPtr<SWidget> MakeViewportToolbar() override;
 	virtual void PopulateViewportOverlays(TSharedRef<SOverlay> Overlay) override;
 
diff --git a/Engine/Source/Editor/UnrealEd/Private/SEditorViewport.cpp b/Engine/Source/Editor/UnrealEd/Private/SEditorViewport.cpp
--- a/Engine/Source/Editor/UnrealEd/Private/SEditorViewport.cpp
+++ b/Engine/Source/Editor/UnrealEd/Private/SEditorViewport.cpp
@@ -26,6 +26,7 @@
 #include "GPUSkinCacheVisualizationMenuCommands.h"
 #include "GPUSkinCache.h"
 #include "Widgets/Colors/SComplexGradient.h"
+#include "Widgets/SBoxPanel.h"
 #include "Modules/ModuleManager.h"
 #include "ISettingsModule.h"
 #include "ShowFlagMenuCommands.h"
@@ -57,20 +58,18 @@ SEditorViewport::~SEditorViewport()
 
 void SEditorViewport::Construct( const FArguments& InArgs )
 {
-	ChildSlot
+	// Create Viewport Widget
+
+	// clang-format off
+	SAssignNew(ViewportWidget, SViewport)
+	.ShowEffectWhenDisabled(false)
+	.EnableGammaCorrection(false) // Scene rendering handles this
+	.AddMetaData(InArgs.MetaData.Num() > 0 ? InArgs.MetaData[0] : MakeShareable(new FTagMetaData(TEXT("LevelEditorViewport"))))
+	.ViewportSize(InArgs._ViewportSize)
 	[
-		SNew(SGlobalPlayWorldActions)
-		[
-			SAssignNew(ViewportWidget, SViewport)
-			.ShowEffectWhenDisabled(false)
-			.EnableGammaCorrection(false) // Scene rendering handles this
-			.AddMetaData(InArgs.MetaData.Num() > 0 ? InArgs.MetaData[0] : MakeShareable(new FTagMetaData(TEXT("LevelEditorViewport"))))
-			.ViewportSize(InArgs._ViewportSize)
-			[
-				SAssignNew(ViewportOverlay, SOverlay)
-			]
-		]
+		SAssignNew(ViewportOverlay, SOverlay)
 	];
+	// clang-format on
 
 	Client = MakeEditorViewportClient();
 
@@ -92,7 +91,8 @@ void SEditorViewport::Construct( const FArguments& InArgs )
 	// Ensure the commands are registered
 	FEditorViewportCommands::Register();
 	BindCommands();
-	
+
+	// clang-format off
 	ViewportOverlay->AddSlot()
 	[
 		SNew(SBorder)
@@ -102,16 +102,17 @@ void SEditorViewport::Construct( const FArguments& InArgs )
 		.Padding(0.0f)
 		.ShowEffectWhenDisabled(false)
 	];
+	// clang-format on
 
-	TSharedPtr<SWidget> ViewportToolbar = MakeViewportToolbar();
-
-	if (ViewportToolbar.IsValid())
+	if (TSharedPtr<SWidget> ViewportToolbar = MakeViewportToolbar())
 	{
 		ViewportOverlay->AddSlot()
+			// clang-format off
 			.VAlign(VAlign_Top)
 			[
 				ViewportToolbar.ToSharedRef()
 			];
+			// clang-format on
 	}
 
 	// This makes a gradient that displays whether or not a viewport is active
@@ -121,6 +122,7 @@ void SEditorViewport::Construct( const FArguments& InArgs )
 
 	static TArray<FLinearColor> GradientStops{ ActiveBorderColorTransparent, ActiveBorderColor, ActiveBorderColorTransparent };
 
+	// clang-format off
 	ViewportOverlay->AddSlot()
 	.VAlign(VAlign_Top)
 	[
@@ -134,7 +136,39 @@ void SEditorViewport::Construct( const FArguments& InArgs )
 			.Orientation(EOrientation::Orient_Vertical)
 		]
 	];
+	// clang-format on
+
+	TSharedPtr<SVerticalBox> VerticalBox;
+	// clang-format off
+	ChildSlot
+	[
+		SAssignNew(VerticalBox, SVerticalBox)
+	];
+	// clang-format on
 
+	// If new toolbar is available, let's add it on top of viewport
+	if (TSharedPtr<SWidget> NewViewportToolbar = BuildViewportToolbar())
+	{
+		// clang-format off
+		VerticalBox->AddSlot()
+		.AutoHeight()
+		[
+			NewViewportToolbar.ToSharedRef()
+		];
+		// clang-format on
+	}
+
+	// clang-format off
+	VerticalBox->AddSlot()
+	[
+		SNew(SGlobalPlayWorldActions)
+		[
+			ViewportWidget.ToSharedRef()
+		]
+	];
+	// clang-format on
+
+	
 	PopulateViewportOverlays(ViewportOverlay.ToSharedRef());
 }
 
diff --git a/Engine/Source/Editor/UnrealEd/Public/SEditorViewport.h b/Engine/Source/Editor/UnrealEd/Public/SEditorViewport.h
--- a/Engine/Source/Editor/UnrealEd/Public/SEditorViewport.h
+++ b/Engine/Source/Editor/UnrealEd/Public/SEditorViewport.h
@@ -151,6 +151,12 @@ protected:
 
 	virtual TSharedRef<FEditorViewportClient> MakeEditorViewportClient() = 0;
 
+	// Implement this to add a viewport toolbar above the viewport that does not cover the viewport.
+	virtual TSharedPtr<SWidget> BuildViewportToolbar()
+	{
+		return TSharedPtr<SWidget>(nullptr);
+	}
+
 	// Implement this to add a viewport toolbar to the inside top of the viewport
 	virtual TSharedPtr<SWidget> MakeViewportToolbar() { return TSharedPtr<SWidget>(nullptr); }
 
