[Viewport Toolbar] Add the LOD submenu to the new Static Mesh viewport toolbar

diff --git a/Engine/Source/Editor/StaticMeshEditor/Private/SStaticMeshEditorViewport.cpp b/Engine/Source/Editor/StaticMeshEditor/Private/SStaticMeshEditorViewport.cpp
--- a/Engine/Source/Editor/StaticMeshEditor/Private/SStaticMeshEditorViewport.cpp
+++ b/Engine/Source/Editor/StaticMeshEditor/Private/SStaticMeshEditorViewport.cpp
@@ -21,6 +21,7 @@
 #include "SEditorViewportToolBarMenu.h"
 #include "Editor.h"
 #include "PreviewProfileController.h"
+#include "StaticMeshEditorViewportToolbarSections.h"
 #include "ToolMenus.h"
 #include "ViewportToolbar/UnrealEdViewportToolbar.h"
 #include "Widgets/Text/SRichTextBlock.h"
@@ -868,6 +869,13 @@ TSharedPtr<SWidget> SStaticMeshEditorViewport::BuildViewportToolbar()
 				RightSection.AddEntry(ViewModesSubmenu);
 			}
 
+			// Add the "LOD" sub menu
+			{
+				FToolMenuEntry LODSubmenu = UE::StaticMeshEditor::CreateLODSubmenu();
+				LODSubmenu.InsertPosition.Position = EToolMenuInsertType::Last;
+				RightSection.AddEntry(LODSubmenu);
+			}
+
 			// Add the "Performance and Scalability" sub menu.
 			{
 				FToolMenuEntry PerformanceAndScalabilitySubmenu = UE::UnrealEd::CreatePerformanceAndScalabilitySubmenu();
diff --git a/Engine/Source/Editor/StaticMeshEditor/Private/SStaticMeshEditorViewportToolBar.cpp b/Engine/Source/Editor/StaticMeshEditor/Private/SStaticMeshEditorViewportToolBar.cpp
--- a/Engine/Source/Editor/StaticMeshEditor/Private/SStaticMeshEditorViewportToolBar.cpp
+++ b/Engine/Source/Editor/StaticMeshEditor/Private/SStaticMeshEditorViewportToolBar.cpp
@@ -20,6 +20,7 @@
 #include "SEditorViewportToolBarMenu.h"
 #include "StaticMeshViewportLODCommands.h"
 #include "PreviewProfileController.h"
+#include "StaticMeshEditorViewportToolbarSections.h"
 
 #define LOCTEXT_NAMESPACE "StaticMeshEditorViewportToolbar"
 
@@ -96,69 +97,18 @@ void SStaticMeshEditorViewportToolbar::ExtendLeftAlignedToolbarSlots(TSharedPtr<
 
 FText SStaticMeshEditorViewportToolbar::GetLODMenuLabel() const
 {
-	FText Label = LOCTEXT("LODMenu_AutoLabel", "LOD Auto");
-
 	TSharedRef<SEditorViewport> BaseViewportRef = GetInfoProvider().GetViewportWidget();
 	TSharedRef<SStaticMeshEditorViewport> ViewportRef = StaticCastSharedRef<SStaticMeshEditorViewport, SEditorViewport>(BaseViewportRef);
 
-	int32 LODSelectionType = ViewportRef->GetLODSelection();
-
-	if (LODSelectionType > 0)
-	{
-		FString TitleLabel = FString::Printf(TEXT("LOD %d"), LODSelectionType - 1);
-		Label = FText::FromString(TitleLabel);
-	}
-
-	return Label;
+	return UE::StaticMeshEditor::GetLODMenuLabel(ViewportRef);
 }
 
 TSharedRef<SWidget> SStaticMeshEditorViewportToolbar::GenerateLODMenu() const
 {
-	const FStaticMeshViewportLODCommands& Actions = FStaticMeshViewportLODCommands::Get();
-
 	TSharedRef<SEditorViewport> BaseViewportRef = GetInfoProvider().GetViewportWidget();
 	TSharedRef<SStaticMeshEditorViewport> ViewportRef = StaticCastSharedRef<SStaticMeshEditorViewport, SEditorViewport>(BaseViewportRef);
 
-	TSharedPtr<FExtender> MenuExtender = GetInfoProvider().GetExtenders();
-
-	const bool bInShouldCloseWindowAfterMenuSelection = true;
-	FMenuBuilder InMenuBuilder(bInShouldCloseWindowAfterMenuSelection, ViewportRef->GetCommandList(), MenuExtender);
-
-	InMenuBuilder.PushCommandList(ViewportRef->GetCommandList().ToSharedRef());
-	if (MenuExtender.IsValid())
-	{ 
-		InMenuBuilder.PushExtender(MenuExtender.ToSharedRef());
-	}
-
-	{
-		// LOD Models
-		InMenuBuilder.BeginSection("StaticMeshViewportPreviewLODs", LOCTEXT("ShowLOD_PreviewLabel", "Preview LODs"));
-		{
-			InMenuBuilder.AddMenuEntry(Actions.LODAuto);
-			InMenuBuilder.AddMenuEntry(Actions.LOD0);
-
-			int32 LODCount = ViewportRef->GetLODModelCount();
-			for (int32 LODId = 1; LODId < LODCount; ++LODId)
-			{
-				FString TitleLabel = FString::Printf(TEXT(" LOD %d"), LODId);
-
-				FUIAction Action(FExecuteAction::CreateSP(ViewportRef, &SStaticMeshEditorViewport::OnSetLODModel, LODId + 1),
-					FCanExecuteAction(),
-					FIsActionChecked::CreateSP(ViewportRef, &SStaticMeshEditorViewport::IsLODModelSelected, LODId + 1));
-
-				InMenuBuilder.AddMenuEntry(FText::FromString(TitleLabel), FText::GetEmpty(), FSlateIcon(), Action, NAME_None, EUserInterfaceActionType::RadioButton);
-			}
-		}
-		InMenuBuilder.EndSection();
-	}
-
-	InMenuBuilder.PopCommandList();
-	if (MenuExtender.IsValid())
-	{
-		InMenuBuilder.PopExtender();
-	}
-
-	return InMenuBuilder.MakeWidget();
+	return UE::StaticMeshEditor::GenerateLODMenuWidget(ViewportRef);
 }
 
 #undef LOCTEXT_NAMESPACE
diff --git a/Engine/Source/Editor/StaticMeshEditor/Private/StaticMeshEditorViewportToolbarSections.cpp b/Engine/Source/Editor/StaticMeshEditor/Private/StaticMeshEditorViewportToolbarSections.cpp
new file mode 100644
--- /dev/null
+++ b/Engine/Source/Editor/StaticMeshEditor/Private/StaticMeshEditorViewportToolbarSections.cpp
@@ -0,0 +1,145 @@
+﻿// Copyright Epic Games, Inc. All Rights Reserved.
+
+#include "StaticMeshEditorViewportToolbarSections.h"
+
+#include "SStaticMeshEditorViewport.h"
+#include "StaticMeshViewportLODCommands.h"
+#include "ToolMenu.h"
+#include "ToolMenuSection.h"
+#include "ViewportToolbar/UnrealEdViewportToolbarContext.h"
+
+#define LOCTEXT_NAMESPACE "StaticMeshEditorViewportToolbarSections"
+
+FText UE::StaticMeshEditor::GetLODMenuLabel(const TSharedPtr<SStaticMeshEditorViewport>& InStaticMeshEditorViewport)
+{
+	FText Label = LOCTEXT("LODMenu_AutoLabel", "LOD Auto");
+
+	if (InStaticMeshEditorViewport)
+	{
+		int32 LODSelectionType = InStaticMeshEditorViewport->GetLODSelection();
+
+		if (LODSelectionType > 0)
+		{
+			FString TitleLabel = FString::Printf(TEXT("LOD %d"), LODSelectionType - 1);
+			Label = FText::FromString(TitleLabel);
+		}
+	}
+
+	return Label;
+}
+
+FToolMenuEntry UE::StaticMeshEditor::CreateLODSubmenu()
+{
+	return FToolMenuEntry::InitDynamicEntry(
+		"DynamicLODOptions",
+		FNewToolMenuSectionDelegate::CreateLambda(
+			[](FToolMenuSection& InDynamicSection) -> void
+			{
+				if (UUnrealEdViewportToolbarContext* const EditorViewportContext =
+						InDynamicSection.FindContext<UUnrealEdViewportToolbarContext>())
+				{
+					TWeakPtr<SStaticMeshEditorViewport> StaticMeshEditorViewportWeak =
+						StaticCastSharedPtr<SStaticMeshEditorViewport>(EditorViewportContext->Viewport.Pin());
+
+					// Label updates based on currently selected LOD
+					const TAttribute<FText> Label = TAttribute<FText>::CreateLambda(
+						[StaticMeshEditorViewportWeak]()
+						{
+							if (TSharedPtr<SStaticMeshEditorViewport> Viewport = StaticMeshEditorViewportWeak.Pin())
+							{
+								return GetLODMenuLabel(Viewport);
+							}
+
+							return LOCTEXT("LODSubmenuLabel", "LOD");
+						}
+					);
+
+					InDynamicSection.AddSubMenu(
+						"LOD",
+						Label,
+						LOCTEXT("LODSubmenuTooltip", ""),
+						FNewToolMenuDelegate::CreateLambda(
+							[StaticMeshEditorViewportWeak](UToolMenu* Submenu) -> void
+							{
+								if (TSharedPtr<SStaticMeshEditorViewport> Viewport = StaticMeshEditorViewportWeak.Pin())
+								{
+									FToolMenuSection& UnnamedSection =
+										Submenu->FindOrAddSection("", LOCTEXT("UnnamedLabel", ""));
+									TSharedRef<SWidget> LODMenuWidget =
+										UE::StaticMeshEditor::GenerateLODMenuWidget(Viewport);
+									FToolMenuEntry LODSubmenu = FToolMenuEntry::InitWidget("LOD", LODMenuWidget, FText());
+
+									UnnamedSection.AddEntry(LODSubmenu);
+								}
+							}
+						)
+					);
+				}
+			}
+		)
+	);
+}
+
+TSharedRef<SWidget> UE::StaticMeshEditor::GenerateLODMenuWidget(const TSharedPtr<SStaticMeshEditorViewport>& InStaticMeshEditorViewport
+)
+{
+	if (!InStaticMeshEditorViewport)
+	{
+		return SNullWidget::NullWidget;
+	}
+
+	const FStaticMeshViewportLODCommands& Actions = FStaticMeshViewportLODCommands::Get();
+
+	TSharedPtr<FExtender> MenuExtender = InStaticMeshEditorViewport->GetExtenders();
+
+	constexpr bool bInShouldCloseWindowAfterMenuSelection = true;
+	FMenuBuilder InMenuBuilder(
+		bInShouldCloseWindowAfterMenuSelection, InStaticMeshEditorViewport->GetCommandList(), MenuExtender
+	);
+
+	InMenuBuilder.PushCommandList(InStaticMeshEditorViewport->GetCommandList().ToSharedRef());
+	if (MenuExtender.IsValid())
+	{
+		InMenuBuilder.PushExtender(MenuExtender.ToSharedRef());
+	}
+
+	{
+		// LOD Models
+		InMenuBuilder.BeginSection("StaticMeshViewportPreviewLODs", LOCTEXT("ShowLOD_PreviewLabel", "Preview LODs"));
+		{
+			InMenuBuilder.AddMenuEntry(Actions.LODAuto);
+			InMenuBuilder.AddMenuEntry(Actions.LOD0);
+
+			int32 LODCount = InStaticMeshEditorViewport->GetLODModelCount();
+			for (int32 LODId = 1; LODId < LODCount; ++LODId)
+			{
+				FString TitleLabel = FString::Printf(TEXT(" LOD %d"), LODId);
+
+				FUIAction Action(
+					FExecuteAction::CreateSP(
+						InStaticMeshEditorViewport.ToSharedRef(), &SStaticMeshEditorViewport::OnSetLODModel, LODId + 1
+					),
+					FCanExecuteAction(),
+					FIsActionChecked::CreateSP(
+						InStaticMeshEditorViewport.ToSharedRef(), &SStaticMeshEditorViewport::IsLODModelSelected, LODId + 1
+					)
+				);
+
+				InMenuBuilder.AddMenuEntry(
+					FText::FromString(TitleLabel), FText::GetEmpty(), FSlateIcon(), Action, NAME_None, EUserInterfaceActionType::RadioButton
+				);
+			}
+		}
+		InMenuBuilder.EndSection();
+	}
+
+	InMenuBuilder.PopCommandList();
+	if (MenuExtender.IsValid())
+	{
+		InMenuBuilder.PopExtender();
+	}
+
+	return InMenuBuilder.MakeWidget();
+}
+
+#undef LOCTEXT_NAMESPACE
diff --git a/Engine/Source/Editor/StaticMeshEditor/Private/StaticMeshEditorViewportToolbarSections.h b/Engine/Source/Editor/StaticMeshEditor/Private/StaticMeshEditorViewportToolbarSections.h
new file mode 100644
--- /dev/null
+++ b/Engine/Source/Editor/StaticMeshEditor/Private/StaticMeshEditorViewportToolbarSections.h
@@ -0,0 +1,16 @@
+﻿// Copyright Epic Games, Inc. All Rights Reserved.
+
+#pragma once
+
+#include "Templates/SharedPointerFwd.h"
+#include "ToolMenuEntry.h"
+
+class SStaticMeshEditorViewport;
+class SWidget;
+
+namespace UE::StaticMeshEditor
+{
+FToolMenuEntry CreateLODSubmenu();
+TSharedRef<SWidget> GenerateLODMenuWidget(const TSharedPtr<SStaticMeshEditorViewport>& InStaticMeshEditorViewport);
+FText GetLODMenuLabel(const TSharedPtr<SStaticMeshEditorViewport>& InStaticMeshEditorViewport);
+}
