Interchange:
- Apply reimport strategy on collisions

StaticMeshEditorSubsystem:
- Set editor-generator flag on the collision

GeomFitUtils:
- Remove setting of editor-generated flag, leaving it to the user to set it

diff --git a/Engine/Plugins/Interchange/Runtime/Source/Import/Private/Mesh/InterchangeStaticMeshFactory.cpp b/Engine/Plugins/Interchange/Runtime/Source/Import/Private/Mesh/InterchangeStaticMeshFactory.cpp
--- a/Engine/Plugins/Interchange/Runtime/Source/Import/Private/Mesh/InterchangeStaticMeshFactory.cpp
+++ b/Engine/Plugins/Interchange/Runtime/Source/Import/Private/Mesh/InterchangeStaticMeshFactory.cpp
@@ -434,12 +434,25 @@ UInterchangeFactoryBase::FImportAssetResult UInterchangeStaticMeshFactory::Impor
 			{
 				if (bReimport)
 				{
+					EReimportStrategyFlags ReimportStrategy = StaticMeshFactoryNode->GetReimportStrategyFlags();
+					switch(ReimportStrategy)
+					{
+					using enum EReimportStrategyFlags;
+					case ApplyNoProperties:	// do nothing
+						break;
+
+					case ApplyPipelineProperties:
+						StaticMesh->GetBodySetup()->AggGeom.EmptyElements();
+						break;
+
+					case ApplyEditorChangedProperties:
 #if WITH_EDITORONLY_DATA
-					StaticMesh->GetBodySetup()->AggGeom.EmptyImportedElements();
-#else
-					StaticMesh->GetBodySetup()->AggGeom.EmptyElements();
+						StaticMesh->GetBodySetup()->AggGeom.EmptyImportedElements();
 #endif
+						break;
+					}
 				}
+
 				bImportedCustomCollision |= ImportBoxCollision(Arguments, StaticMesh, LodDataNode);
 				bImportedCustomCollision |= ImportCapsuleCollision(Arguments, StaticMesh, LodDataNode);
 				bImportedCustomCollision |= ImportSphereCollision(Arguments, StaticMesh, LodDataNode);
diff --git a/Engine/Source/Editor/StaticMeshEditor/Private/StaticMeshEditorSubsystem.cpp b/Engine/Source/Editor/StaticMeshEditor/Private/StaticMeshEditorSubsystem.cpp
--- a/Engine/Source/Editor/StaticMeshEditor/Private/StaticMeshEditorSubsystem.cpp
+++ b/Engine/Source/Editor/StaticMeshEditor/Private/StaticMeshEditorSubsystem.cpp
@@ -1162,46 +1162,78 @@ int32 UStaticMeshEditorSubsystem::AddSimpleCollisionsWithNotification(UStaticMes
 	case EScriptCollisionShapeType::Box:
 	{
 		PrimIndex = GenerateBoxAsSimpleCollision(StaticMesh);
+		if(PrimIndex != INDEX_NONE)
+		{
+			StaticMesh->GetBodySetup()->AggGeom.BoxElems[PrimIndex].bIsGenerated = true;
+		}
 		break;
 	}
 	case EScriptCollisionShapeType::Sphere:
 	{
 		PrimIndex = GenerateSphereAsSimpleCollision(StaticMesh);
+		if(PrimIndex != INDEX_NONE)
+		{
+			StaticMesh->GetBodySetup()->AggGeom.SphereElems[PrimIndex].bIsGenerated = true;
+		}
 		break;
 	}
 	case EScriptCollisionShapeType::Capsule:
 	{
 		PrimIndex = GenerateSphylAsSimpleCollision(StaticMesh);
+		if(PrimIndex != INDEX_NONE)
+		{
+			StaticMesh->GetBodySetup()->AggGeom.SphylElems[PrimIndex].bIsGenerated = true;
+		}
 		break;
 	}
 	case EScriptCollisionShapeType::NDOP10_X:
 	{
 		TArray<FVector>	DirArray(KDopDir10X, 10);
 		PrimIndex = GenerateKDopAsSimpleCollision(StaticMesh, DirArray);
+		if(PrimIndex != INDEX_NONE)
+		{
+			StaticMesh->GetBodySetup()->AggGeom.ConvexElems[PrimIndex].bIsGenerated = true;
+		}
 		break;
 	}
 	case EScriptCollisionShapeType::NDOP10_Y:
 	{
 		TArray<FVector>	DirArray(KDopDir10Y, 10);
 		PrimIndex = GenerateKDopAsSimpleCollision(StaticMesh, DirArray);
+		if(PrimIndex != INDEX_NONE)
+		{
+			StaticMesh->GetBodySetup()->AggGeom.ConvexElems[PrimIndex].bIsGenerated = true;
+		}
 		break;
 	}
 	case EScriptCollisionShapeType::NDOP10_Z:
 	{
 		TArray<FVector>	DirArray(KDopDir10Z, 10);
 		PrimIndex = GenerateKDopAsSimpleCollision(StaticMesh, DirArray);
+		if(PrimIndex != INDEX_NONE)
+		{
+			StaticMesh->GetBodySetup()->AggGeom.ConvexElems[PrimIndex].bIsGenerated = true;
+		}
 		break;
 	}
 	case EScriptCollisionShapeType::NDOP18:
 	{
 		TArray<FVector>	DirArray(KDopDir18, 18);
 		PrimIndex = GenerateKDopAsSimpleCollision(StaticMesh, DirArray);
+		if(PrimIndex != INDEX_NONE)
+		{
+			StaticMesh->GetBodySetup()->AggGeom.ConvexElems[PrimIndex].bIsGenerated = true;
+		}
 		break;
 	}
 	case EScriptCollisionShapeType::NDOP26:
 	{
 		TArray<FVector>	DirArray(KDopDir26, 26);
 		PrimIndex = GenerateKDopAsSimpleCollision(StaticMesh, DirArray);
+		if(PrimIndex != INDEX_NONE)
+		{
+			StaticMesh->GetBodySetup()->AggGeom.ConvexElems[PrimIndex].bIsGenerated = true;
+		}
 		break;
 	}
 	}
diff --git a/Engine/Source/Editor/UnrealEd/Private/GeomFitUtils.cpp b/Engine/Source/Editor/UnrealEd/Private/GeomFitUtils.cpp
--- a/Engine/Source/Editor/UnrealEd/Private/GeomFitUtils.cpp
+++ b/Engine/Source/Editor/UnrealEd/Private/GeomFitUtils.cpp
@@ -108,7 +108,6 @@ int32 GenerateKDopAsSimpleCollision(UStaticMesh* StaticMesh, const TArray<FVecto
 	bs->InvalidatePhysicsData();
 
 	FKConvexElem ConvexElem;
-	ConvexElem.bIsGenerated = true;
 	ConvexElem.VertexData = HullVertices;
 	// Note: UpdateElemBox also computes the convex hull indices
 	ConvexElem.UpdateElemBox();
@@ -161,7 +160,6 @@ int32 GenerateBoxAsSimpleCollision(UStaticMesh* StaticMesh, bool bUpdate)
 	bs->InvalidatePhysicsData();
 
 	FKBoxElem BoxElem;
-	BoxElem.bIsGenerated = true;
 	BoxElem.Center = Center;
 	BoxElem.X = Extents.X * 2.0f;
 	BoxElem.Y = Extents.Y * 2.0f;
@@ -366,7 +364,6 @@ int32 GenerateSphereAsSimpleCollision(UStaticMesh* StaticMesh, bool bUpdate)
 	bs->InvalidatePhysicsData();
 
 	FKSphereElem SphereElem;
-	SphereElem.bIsGenerated = true;
 	SphereElem.Center = bestSphere.Center;
 	SphereElem.Radius = bestSphere.W;
 	bs->AggGeom.SphereElems.Add(SphereElem);
@@ -517,7 +514,6 @@ int32 GenerateSphylAsSimpleCollision(UStaticMesh* StaticMesh, bool bUpdate)
 	bs->InvalidatePhysicsData();
 
 	FKSphylElem SphylElem;
-	SphylElem.bIsGenerated = true;
 	SphylElem.Center = sphere.Center;
 	SphylElem.Rotation = rotation;
 	SphylElem.Radius = sphere.W;
