[Viewport Toolbar] Add "Show" submenu to Static Mesh editor's new viewport toolbar

Moves the existing code creating the Static Mesh Show Menu from the old toolbar, and wraps it so that both the new and old toolbar can use it.

diff --git a/Engine/Source/Editor/StaticMeshEditor/Private/SStaticMeshEditorViewport.cpp b/Engine/Source/Editor/StaticMeshEditor/Private/SStaticMeshEditorViewport.cpp
--- a/Engine/Source/Editor/StaticMeshEditor/Private/SStaticMeshEditorViewport.cpp
+++ b/Engine/Source/Editor/StaticMeshEditor/Private/SStaticMeshEditorViewport.cpp
@@ -871,6 +871,13 @@ TSharedPtr<SWidget> SStaticMeshEditorViewport::BuildViewportToolbar()
 				RightSection.AddEntry(ViewModesSubmenu);
 			}
 
+			// Add the "Show" sub menu
+			{
+				FToolMenuEntry ShowSubmenu = UE::StaticMeshEditor::CreateShowSubmenu();
+				ShowSubmenu.InsertPosition.Position = EToolMenuInsertType::Last;
+				RightSection.AddEntry(ShowSubmenu);
+			}
+
 			// Add the "LOD" sub menu
 			{
 				FToolMenuEntry LODSubmenu = UE::StaticMeshEditor::CreateLODSubmenu();
diff --git a/Engine/Source/Editor/StaticMeshEditor/Private/SStaticMeshEditorViewportToolBar.cpp b/Engine/Source/Editor/StaticMeshEditor/Private/SStaticMeshEditorViewportToolBar.cpp
--- a/Engine/Source/Editor/StaticMeshEditor/Private/SStaticMeshEditorViewportToolBar.cpp
+++ b/Engine/Source/Editor/StaticMeshEditor/Private/SStaticMeshEditorViewportToolBar.cpp
@@ -36,41 +36,10 @@ void SStaticMeshEditorViewportToolbar::Construct(const FArguments& InArgs, TShar
 // SCommonEditorViewportToolbarBase interface
 TSharedRef<SWidget> SStaticMeshEditorViewportToolbar::GenerateShowMenu() const
 {
-	GetInfoProvider().OnFloatingButtonClicked();
-
-	TSharedRef<SEditorViewport> ViewportRef = GetInfoProvider().GetViewportWidget();
-
-	const bool bInShouldCloseWindowAfterMenuSelection = true;
-	FMenuBuilder ShowMenuBuilder(bInShouldCloseWindowAfterMenuSelection, ViewportRef->GetCommandList());
-	{
-		auto Commands = FStaticMeshEditorCommands::Get();
-
-		ShowMenuBuilder.AddMenuEntry(Commands.SetShowNaniteFallback);
-		ShowMenuBuilder.AddMenuEntry(Commands.SetShowDistanceField);
-
-		ShowMenuBuilder.BeginSection("MeshComponents", LOCTEXT("MeshComponments", "Mesh Components"));
-		ShowMenuBuilder.AddMenuEntry(Commands.SetShowSockets);
-		ShowMenuBuilder.AddMenuEntry(Commands.SetShowVertices);
-		ShowMenuBuilder.AddMenuEntry(Commands.SetShowVertexColor);
-		ShowMenuBuilder.AddMenuEntry(Commands.SetShowNormals);
-		ShowMenuBuilder.AddMenuEntry(Commands.SetShowTangents);
-		ShowMenuBuilder.AddMenuEntry(Commands.SetShowBinormals);
-		ShowMenuBuilder.AddMenuSeparator();
-
-		ShowMenuBuilder.AddMenuEntry(Commands.SetShowPivot);
-		ShowMenuBuilder.AddMenuEntry(Commands.SetShowGrid);
-		ShowMenuBuilder.AddMenuEntry(Commands.SetShowBounds);
-		ShowMenuBuilder.AddMenuEntry(Commands.SetShowSimpleCollision);
-		ShowMenuBuilder.AddMenuEntry(Commands.SetShowComplexCollision);
-		ShowMenuBuilder.AddMenuEntry(Commands.SetShowPhysicalMaterialMasks);
-
-
-
-		//ShowMenuBuilder.AddMenuSeparator();
-		//ShowMenuBuilder.AddMenuEntry(Commands.SetShowMeshEdges);
-	}
+	TSharedRef<SEditorViewport> BaseViewportRef = GetInfoProvider().GetViewportWidget();
+	TSharedRef<SStaticMeshEditorViewport> ViewportRef = StaticCastSharedRef<SStaticMeshEditorViewport, SEditorViewport>(BaseViewportRef);
 
-	return ShowMenuBuilder.MakeWidget();
+	return UE::StaticMeshEditor::GenerateShowMenuWidget(ViewportRef);
 }
 
 // SCommonEditorViewportToolbarBase interface
diff --git a/Engine/Source/Editor/StaticMeshEditor/Private/StaticMeshEditorViewportToolbarSections.cpp b/Engine/Source/Editor/StaticMeshEditor/Private/StaticMeshEditorViewportToolbarSections.cpp
--- a/Engine/Source/Editor/StaticMeshEditor/Private/StaticMeshEditorViewportToolbarSections.cpp
+++ b/Engine/Source/Editor/StaticMeshEditor/Private/StaticMeshEditorViewportToolbarSections.cpp
@@ -3,9 +3,12 @@
 #include "StaticMeshEditorViewportToolbarSections.h"
 
 #include "SStaticMeshEditorViewport.h"
+#include "StaticMeshEditorActions.h"
 #include "StaticMeshViewportLODCommands.h"
 #include "ToolMenu.h"
 #include "ToolMenuSection.h"
+#include "ToolMenus.h"
+#include "ViewportToolbar/UnrealEdViewportToolbar.h"
 #include "ViewportToolbar/UnrealEdViewportToolbarContext.h"
 
 #define LOCTEXT_NAMESPACE "StaticMeshEditorViewportToolbarSections"
@@ -63,8 +66,7 @@ FToolMenuEntry UE::StaticMeshEditor::CreateLODSubmenu()
 							{
 								if (TSharedPtr<SStaticMeshEditorViewport> Viewport = StaticMeshEditorViewportWeak.Pin())
 								{
-									FToolMenuSection& UnnamedSection =
-										Submenu->FindOrAddSection("", FText());
+									FToolMenuSection& UnnamedSection = Submenu->FindOrAddSection(NAME_None, FText());
 									TSharedRef<SWidget> LODMenuWidget =
 										UE::StaticMeshEditor::GenerateLODMenuWidget(Viewport);
 									FToolMenuEntry LODSubmenu = FToolMenuEntry::InitWidget("LOD", LODMenuWidget, FText());
@@ -80,6 +82,21 @@ FToolMenuEntry UE::StaticMeshEditor::CreateLODSubmenu()
 	);
 }
 
+FToolMenuEntry UE::StaticMeshEditor::CreateShowSubmenu()
+{
+	return FToolMenuEntry::InitSubMenu(
+		"Show",
+		LOCTEXT("ShowSubmenuLabel", "Show"),
+		FText(),
+		FNewToolMenuDelegate::CreateLambda(
+			[](UToolMenu* Submenu) -> void
+			{
+				FillShowSubmenu(Submenu);
+			}
+		)
+	);
+}
+
 TSharedRef<SWidget> UE::StaticMeshEditor::GenerateLODMenuWidget(const TSharedPtr<SStaticMeshEditorViewport>& InStaticMeshEditorViewport
 )
 {
@@ -142,4 +159,83 @@ TSharedRef<SWidget> UE::StaticMeshEditor::GenerateLODMenuWidget(const TSharedPtr
 	return InMenuBuilder.MakeWidget();
 }
 
+TSharedRef<SWidget> UE::StaticMeshEditor::GenerateShowMenuWidget(const TSharedPtr<SStaticMeshEditorViewport>& InStaticMeshEditorViewport
+)
+{
+	if (!InStaticMeshEditorViewport)
+	{
+		return SNullWidget::NullWidget;
+	}
+
+	InStaticMeshEditorViewport->OnFloatingButtonClicked();
+
+	// We generate a menu via UToolMenus, so we can use FillShowSubmenu call from both old and new toolbar
+	FName OldShowMenuName = "StaticMesh.OldViewportToolbar.Show";
+
+	if (!UToolMenus::Get()->IsMenuRegistered(OldShowMenuName))
+	{
+		UToolMenu* Menu = UToolMenus::Get()->RegisterMenu(OldShowMenuName, NAME_None, EMultiBoxType::Menu, false);
+		Menu->AddDynamicSection(
+			"BaseSection",
+			FNewToolMenuDelegate::CreateLambda(
+				[](UToolMenu* InMenu)
+				{
+					FillShowSubmenu(InMenu);
+				}
+			)
+		);
+	}
+
+	FToolMenuContext MenuContext;
+	{
+		MenuContext.AppendCommandList(InStaticMeshEditorViewport->GetCommandList());
+
+		// Add the UnrealEd viewport toolbar context.
+		{
+			UUnrealEdViewportToolbarContext* const ContextObject =
+				UE::UnrealEd::CreateViewportToolbarDefaultContext(InStaticMeshEditorViewport);
+
+			MenuContext.AddObject(ContextObject);
+		}
+	}
+
+	return UToolMenus::Get()->GenerateWidget(OldShowMenuName, MenuContext);
+}
+
+void UE::StaticMeshEditor::FillShowSubmenu(UToolMenu* InMenu)
+{
+	if (UUnrealEdViewportToolbarContext* const EditorViewportContext =
+			InMenu->FindContext<UUnrealEdViewportToolbarContext>())
+	{
+		if (TSharedPtr<SStaticMeshEditorViewport> StaticMeshEditorViewport =
+				StaticCastSharedPtr<SStaticMeshEditorViewport>(EditorViewportContext->Viewport.Pin()))
+		{
+			FToolMenuSection& UnnamedSection = InMenu->FindOrAddSection(NAME_None);
+			const FStaticMeshEditorCommands& Commands = FStaticMeshEditorCommands::Get();
+
+			UnnamedSection.AddMenuEntry(Commands.SetShowNaniteFallback);
+			UnnamedSection.AddMenuEntry(Commands.SetShowDistanceField);
+
+			FToolMenuSection& MeshComponentsSection =
+				InMenu->FindOrAddSection("MeshComponents", LOCTEXT("MeshComponments", "Mesh Components"));
+
+			MeshComponentsSection.AddMenuEntry(Commands.SetShowSockets);
+			MeshComponentsSection.AddMenuEntry(Commands.SetShowVertices);
+			MeshComponentsSection.AddMenuEntry(Commands.SetShowVertexColor);
+			MeshComponentsSection.AddMenuEntry(Commands.SetShowNormals);
+			MeshComponentsSection.AddMenuEntry(Commands.SetShowTangents);
+			MeshComponentsSection.AddMenuEntry(Commands.SetShowBinormals);
+
+			MeshComponentsSection.AddSeparator(NAME_None);
+
+			MeshComponentsSection.AddMenuEntry(Commands.SetShowPivot);
+			MeshComponentsSection.AddMenuEntry(Commands.SetShowGrid);
+			MeshComponentsSection.AddMenuEntry(Commands.SetShowBounds);
+			MeshComponentsSection.AddMenuEntry(Commands.SetShowSimpleCollision);
+			MeshComponentsSection.AddMenuEntry(Commands.SetShowComplexCollision);
+			MeshComponentsSection.AddMenuEntry(Commands.SetShowPhysicalMaterialMasks);
+		}
+	}
+}
+
 #undef LOCTEXT_NAMESPACE
diff --git a/Engine/Source/Editor/StaticMeshEditor/Private/StaticMeshEditorViewportToolbarSections.h b/Engine/Source/Editor/StaticMeshEditor/Private/StaticMeshEditorViewportToolbarSections.h
--- a/Engine/Source/Editor/StaticMeshEditor/Private/StaticMeshEditorViewportToolbarSections.h
+++ b/Engine/Source/Editor/StaticMeshEditor/Private/StaticMeshEditorViewportToolbarSections.h
@@ -11,6 +11,12 @@ class SWidget;
 namespace UE::StaticMeshEditor
 {
 FToolMenuEntry CreateLODSubmenu();
+FToolMenuEntry CreateShowSubmenu();
+
 TSharedRef<SWidget> GenerateLODMenuWidget(const TSharedPtr<SStaticMeshEditorViewport>& InStaticMeshEditorViewport);
+TSharedRef<SWidget> GenerateShowMenuWidget(const TSharedPtr<SStaticMeshEditorViewport>& InStaticMeshEditorViewport);
+
+void FillShowSubmenu(UToolMenu *InMenu);
+
 FText GetLODMenuLabel(const TSharedPtr<SStaticMeshEditorViewport>& InStaticMeshEditorViewport);
 }
