[Viewport Toolbar] Convert Static Mesh Editor LOD menu to use UToolMenu

diff --git a/Engine/Source/Editor/StaticMeshEditor/Private/StaticMeshEditorViewportToolbarSections.cpp b/Engine/Source/Editor/StaticMeshEditor/Private/StaticMeshEditorViewportToolbarSections.cpp
--- a/Engine/Source/Editor/StaticMeshEditor/Private/StaticMeshEditorViewportToolbarSections.cpp
+++ b/Engine/Source/Editor/StaticMeshEditor/Private/StaticMeshEditorViewportToolbarSections.cpp
@@ -13,6 +13,53 @@
 
 #define LOCTEXT_NAMESPACE "StaticMeshEditorViewportToolbarSections"
 
+namespace UE::StaticMeshEditor::Private
+{
+void PopulateLODSubmenu(UToolMenu* InMenu)
+{
+	if (UUnrealEdViewportToolbarContext* const EditorViewportContext =
+			InMenu->FindContext<UUnrealEdViewportToolbarContext>())
+	{
+		if (TSharedPtr<SStaticMeshEditorViewport> StaticMeshEditorViewport =
+				StaticCastSharedPtr<SStaticMeshEditorViewport>(EditorViewportContext->Viewport.Pin()))
+		{
+			const FStaticMeshViewportLODCommands& Actions = FStaticMeshViewportLODCommands::Get();
+
+			FToolMenuSection& PreviewLODSection =
+				InMenu->FindOrAddSection("StaticMeshViewportPreviewLODs", LOCTEXT("ShowLOD_PreviewLabel", "Preview LODs"));
+
+			PreviewLODSection.AddMenuEntry(Actions.LODAuto);
+			PreviewLODSection.AddMenuEntry(Actions.LOD0);
+
+			int32 LODCount = StaticMeshEditorViewport->GetLODModelCount();
+			for (int32 LODId = 1; LODId < LODCount; ++LODId)
+			{
+				FString TitleLabel = FString::Printf(TEXT(" LOD %d"), LODId);
+
+				FUIAction Action(
+					FExecuteAction::CreateSP(
+						StaticMeshEditorViewport.ToSharedRef(), &SStaticMeshEditorViewport::OnSetLODModel, LODId + 1
+					),
+					FCanExecuteAction(),
+					FIsActionChecked::CreateSP(
+						StaticMeshEditorViewport.ToSharedRef(), &SStaticMeshEditorViewport::IsLODModelSelected, LODId + 1
+					)
+				);
+
+				PreviewLODSection.AddMenuEntry(
+					FName(TitleLabel),
+					FText::FromString(TitleLabel),
+					FText::GetEmpty(),
+					FSlateIcon(),
+					Action,
+					EUserInterfaceActionType::RadioButton
+				);
+			}
+		}
+	}
+}
+} // namespace UE::StaticMeshEditor::Private
+
 FText UE::StaticMeshEditor::GetLODMenuLabel(const TSharedPtr<SStaticMeshEditorViewport>& InStaticMeshEditorViewport)
 {
 	FText Label = LOCTEXT("LODMenu_AutoLabel", "LOD Auto");
@@ -62,17 +109,9 @@ FToolMenuEntry UE::StaticMeshEditor::CreateLODSubmenu()
 						Label,
 						FText(),
 						FNewToolMenuDelegate::CreateLambda(
-							[StaticMeshEditorViewportWeak](UToolMenu* Submenu) -> void
+							[](UToolMenu* Submenu) -> void
 							{
-								if (TSharedPtr<SStaticMeshEditorViewport> Viewport = StaticMeshEditorViewportWeak.Pin())
-								{
-									FToolMenuSection& UnnamedSection = Submenu->FindOrAddSection(NAME_None, FText());
-									TSharedRef<SWidget> LODMenuWidget =
-										UE::StaticMeshEditor::GenerateLODMenuWidget(Viewport);
-									FToolMenuEntry LODSubmenu = FToolMenuEntry::InitWidget("LOD", LODMenuWidget, FText());
-
-									UnnamedSection.AddEntry(LODSubmenu);
-								}
+								UE::StaticMeshEditor::Private::PopulateLODSubmenu(Submenu);
 							}
 						)
 					);
@@ -105,58 +144,39 @@ TSharedRef<SWidget> UE::StaticMeshEditor::GenerateLODMenuWidget(const TSharedPtr
 		return SNullWidget::NullWidget;
 	}
 
-	const FStaticMeshViewportLODCommands& Actions = FStaticMeshViewportLODCommands::Get();
-
-	TSharedPtr<FExtender> MenuExtender = InStaticMeshEditorViewport->GetExtenders();
-
-	constexpr bool bInShouldCloseWindowAfterMenuSelection = true;
-	FMenuBuilder InMenuBuilder(
-		bInShouldCloseWindowAfterMenuSelection, InStaticMeshEditorViewport->GetCommandList(), MenuExtender
-	);
+	// We generate a menu via UToolMenus, so we can use FillShowSubmenu call from both old and new toolbar
+	FName OldShowMenuName = "StaticMeshEditor.OldViewportToolbar.LODMenu";
 
-	InMenuBuilder.PushCommandList(InStaticMeshEditorViewport->GetCommandList().ToSharedRef());
-	if (MenuExtender.IsValid())
+	if (!UToolMenus::Get()->IsMenuRegistered(OldShowMenuName))
 	{
-		InMenuBuilder.PushExtender(MenuExtender.ToSharedRef());
+		if (UToolMenu* Menu = UToolMenus::Get()->RegisterMenu(OldShowMenuName, NAME_None, EMultiBoxType::Menu, false))
+		{
+			Menu->AddDynamicSection(
+				"BaseSection",
+				FNewToolMenuDelegate::CreateLambda(
+					[](UToolMenu* InMenu)
+					{
+						UE::StaticMeshEditor::Private::PopulateLODSubmenu(InMenu);
+					}
+				)
+			);
+		}
 	}
 
+	FToolMenuContext MenuContext;
 	{
-		// LOD Models
-		InMenuBuilder.BeginSection("StaticMeshViewportPreviewLODs", LOCTEXT("ShowLOD_PreviewLabel", "Preview LODs"));
-		{
-			InMenuBuilder.AddMenuEntry(Actions.LODAuto);
-			InMenuBuilder.AddMenuEntry(Actions.LOD0);
-
-			int32 LODCount = InStaticMeshEditorViewport->GetLODModelCount();
-			for (int32 LODId = 1; LODId < LODCount; ++LODId)
-			{
-				FString TitleLabel = FString::Printf(TEXT(" LOD %d"), LODId);
+		MenuContext.AppendCommandList(InStaticMeshEditorViewport->GetCommandList());
 
-				FUIAction Action(
-					FExecuteAction::CreateSP(
-						InStaticMeshEditorViewport.ToSharedRef(), &SStaticMeshEditorViewport::OnSetLODModel, LODId + 1
-					),
-					FCanExecuteAction(),
-					FIsActionChecked::CreateSP(
-						InStaticMeshEditorViewport.ToSharedRef(), &SStaticMeshEditorViewport::IsLODModelSelected, LODId + 1
-					)
-				);
+		// Add the UnrealEd viewport toolbar context.
+		{
+			UUnrealEdViewportToolbarContext* const ContextObject =
+				UE::UnrealEd::CreateViewportToolbarDefaultContext(InStaticMeshEditorViewport);
 
-				InMenuBuilder.AddMenuEntry(
-					FText::FromString(TitleLabel), FText::GetEmpty(), FSlateIcon(), Action, NAME_None, EUserInterfaceActionType::RadioButton
-				);
-			}
+			MenuContext.AddObject(ContextObject);
 		}
-		InMenuBuilder.EndSection();
 	}
 
-	InMenuBuilder.PopCommandList();
-	if (MenuExtender.IsValid())
-	{
-		InMenuBuilder.PopExtender();
-	}
-
-	return InMenuBuilder.MakeWidget();
+	return UToolMenus::Get()->GenerateWidget(OldShowMenuName, MenuContext);
 }
 
 TSharedRef<SWidget> UE::StaticMeshEditor::GenerateShowMenuWidget(const TSharedPtr<SStaticMeshEditorViewport>& InStaticMeshEditorViewport
