Expose Static Mesh LOD Section Properties to Python

diff --git a/Engine/Source/Editor/StaticMeshEditor/Private/StaticMeshEditorSubsystem.cpp b/Engine/Source/Editor/StaticMeshEditor/Private/StaticMeshEditorSubsystem.cpp
--- a/Engine/Source/Editor/StaticMeshEditor/Private/StaticMeshEditorSubsystem.cpp
+++ b/Engine/Source/Editor/StaticMeshEditor/Private/StaticMeshEditorSubsystem.cpp
@@ -591,7 +591,7 @@ void UStaticMeshEditorSubsystem::SetLodBuildSettings(UStaticMesh* StaticMesh, co
 	}
 }
 
-FName UStaticMeshEditorSubsystem::GetLODGroup(UStaticMesh* StaticMesh)
+FName UStaticMeshEditorSubsystem::GetLODGroup(const UStaticMesh* StaticMesh)
 {
 	return StaticMesh->LODGroup;
 }
@@ -904,7 +904,7 @@ int32 UStaticMeshEditorSubsystem::SetLodFromStaticMesh(UStaticMesh* DestinationS
 	return DestinationLodIndex;
 }
 
-int32 UStaticMeshEditorSubsystem::GetLodCount(UStaticMesh* StaticMesh)
+int32 UStaticMeshEditorSubsystem::GetLodCount(const UStaticMesh* StaticMesh)
 {
 	TGuardValue<bool> UnattendedScriptGuard(GIsRunningUnattendedScript, true);
 
@@ -971,7 +971,7 @@ bool UStaticMeshEditorSubsystem::RemoveLods(UStaticMesh* StaticMesh)
 	return true;
 }
 
-TArray<float> UStaticMeshEditorSubsystem::GetLodScreenSizes(UStaticMesh* StaticMesh)
+TArray<float> UStaticMeshEditorSubsystem::GetLodScreenSizes(const UStaticMesh* StaticMesh)
 {
 	TGuardValue<bool> UnattendedScriptGuard(GIsRunningUnattendedScript, true);
 
@@ -1085,7 +1085,7 @@ bool UStaticMeshEditorSubsystem::SetLodScreenSizes(UStaticMesh* StaticMesh, cons
 }
 
 
-FMeshNaniteSettings UStaticMeshEditorSubsystem::GetNaniteSettings(UStaticMesh* StaticMesh)
+FMeshNaniteSettings UStaticMeshEditorSubsystem::GetNaniteSettings(const UStaticMesh* StaticMesh)
 {
 	if (StaticMesh)
 	{
@@ -1253,7 +1253,7 @@ int32 UStaticMeshEditorSubsystem::AddSimpleCollisionsWithNotification(UStaticMes
 	return PrimIndex;
 }
 
-int32 UStaticMeshEditorSubsystem::GetSimpleCollisionCount(UStaticMesh* StaticMesh)
+int32 UStaticMeshEditorSubsystem::GetSimpleCollisionCount(const UStaticMesh* StaticMesh)
 {
 	TGuardValue<bool> UnattendedScriptGuard(GIsRunningUnattendedScript, true);
 
@@ -1281,7 +1281,7 @@ int32 UStaticMeshEditorSubsystem::GetSimpleCollisionCount(UStaticMesh* StaticMes
 	return Count;
 }
 
-TEnumAsByte<ECollisionTraceFlag> UStaticMeshEditorSubsystem::GetCollisionComplexity(UStaticMesh* StaticMesh)
+TEnumAsByte<ECollisionTraceFlag> UStaticMeshEditorSubsystem::GetCollisionComplexity(const UStaticMesh* StaticMesh)
 {
 	TGuardValue<bool> UnattendedScriptGuard(GIsRunningUnattendedScript, true);
 
@@ -1304,7 +1304,7 @@ TEnumAsByte<ECollisionTraceFlag> UStaticMeshEditorSubsystem::GetCollisionComplex
 	return ECollisionTraceFlag::CTF_UseDefault;
 }
 
-int32 UStaticMeshEditorSubsystem::GetConvexCollisionCount(UStaticMesh* StaticMesh)
+int32 UStaticMeshEditorSubsystem::GetConvexCollisionCount(const UStaticMesh* StaticMesh)
 {
 	TGuardValue<bool> UnattendedScriptGuard(GIsRunningUnattendedScript, true);
 
@@ -1507,111 +1507,183 @@ bool UStaticMeshEditorSubsystem::RemoveCollisionsWithNotification(UStaticMesh* S
 	return true;
 }
 
-void UStaticMeshEditorSubsystem::EnableSectionCollision(UStaticMesh* StaticMesh, bool bCollisionEnabled, int32 LODIndex, int32 SectionIndex)
+bool UStaticMeshEditorSubsystem::CanModifyStaticMeshSection(const UStaticMesh* StaticMesh, int32 LODIndex, int32 SectionIndex, const FString& LogFunctionName)
 {
-	TGuardValue<bool> UnattendedScriptGuard(GIsRunningUnattendedScript, true);
-
 	if (!EditorScriptingHelpers::CheckIfInEditorAndPIE())
 	{
-		return;
+		return false;
 	}
 
 	if (StaticMesh == nullptr)
 	{
-		UE_LOG(LogStaticMeshEditorSubsystem, Error, TEXT("EnableSectionCollision: The StaticMesh is null."));
-		return;
+		UE_LOG(LogStaticMeshEditorSubsystem, Error, TEXT("%s: The StaticMesh is null."), *LogFunctionName);
+		return false;
 	}
 
 	if (LODIndex >= StaticMesh->GetNumLODs())
 	{
-		UE_LOG(LogStaticMeshEditorSubsystem, Error, TEXT("EnableSectionCollision: Invalid LOD index %d (of %d)."), LODIndex, StaticMesh->GetNumLODs());
-		return;
+		UE_LOG(LogStaticMeshEditorSubsystem, Error, TEXT("%s: Invalid LOD index %d (of %d)."), *LogFunctionName, LODIndex, StaticMesh->GetNumLODs());
+		return false;
 	}
 
 	if (SectionIndex >= StaticMesh->GetNumSections(LODIndex))
 	{
-		UE_LOG(LogStaticMeshEditorSubsystem, Error, TEXT("EnableSectionCollision: Invalid section index %d (of %d)."), SectionIndex, StaticMesh->GetNumSections(LODIndex));
-		return;
+		UE_LOG(LogStaticMeshEditorSubsystem, Error, TEXT("%s: Invalid section index %d (of %d)."), *LogFunctionName, SectionIndex, StaticMesh->GetNumSections(LODIndex));
+		return false;
 	}
 
-	StaticMesh->Modify();
+	return true;
+}
 
-	FMeshSectionInfo SectionInfo = StaticMesh->GetSectionInfoMap().Get(LODIndex, SectionIndex);
+void UStaticMeshEditorSubsystem::EnableSectionCastShadow(UStaticMesh* StaticMesh, bool bCastShadow, int32 LODIndex, int32 SectionIndex)
+{
+	TGuardValue<bool> UnattendedScriptGuard(GIsRunningUnattendedScript, true);
 
-	SectionInfo.bEnableCollision = bCollisionEnabled;
+	if (CanModifyStaticMeshSection(StaticMesh, LODIndex, SectionIndex, "EnableSectionCastShadow"))
+	{
+		StaticMesh->Modify();
 
-	StaticMesh->GetSectionInfoMap().Set(LODIndex, SectionIndex, SectionInfo);
+		FMeshSectionInfo SectionInfo = StaticMesh->GetSectionInfoMap().Get(LODIndex, SectionIndex);
 
-	StaticMesh->PostEditChange();
+		SectionInfo.bCastShadow = bCastShadow;
+
+		StaticMesh->GetSectionInfoMap().Set(LODIndex, SectionIndex, SectionInfo);
+
+		StaticMesh->PostEditChange();
+	}
 }
 
-bool UStaticMeshEditorSubsystem::IsSectionCollisionEnabled(UStaticMesh* StaticMesh, int32 LODIndex, int32 SectionIndex)
+bool UStaticMeshEditorSubsystem::IsSectionCastShadowEnabled(const UStaticMesh* StaticMesh, int32 LODIndex, int32 SectionIndex)
 {
 	TGuardValue<bool> UnattendedScriptGuard(GIsRunningUnattendedScript, true);
-
-	if (!EditorScriptingHelpers::CheckIfInEditorAndPIE())
+	
+	if (CanModifyStaticMeshSection(StaticMesh, LODIndex, SectionIndex, "IsSectionCastShadowEnabled"))
 	{
-		return false;
+		FMeshSectionInfo SectionInfo = StaticMesh->GetSectionInfoMap().Get(LODIndex, SectionIndex);
+		return SectionInfo.bCastShadow;
 	}
 
-	if (StaticMesh == nullptr)
+	return false;
+}
+
+void UStaticMeshEditorSubsystem::EnableSectionCollision(UStaticMesh* StaticMesh, bool bCollisionEnabled, int32 LODIndex, int32 SectionIndex)
+{
+	TGuardValue<bool> UnattendedScriptGuard(GIsRunningUnattendedScript, true);
+
+	if (CanModifyStaticMeshSection(StaticMesh, LODIndex, SectionIndex, "EnableSectionCollision"))
 	{
-		UE_LOG(LogStaticMeshEditorSubsystem, Error, TEXT("IsSectionCollisionEnabled: The StaticMesh is null."));
-		return false;
+		StaticMesh->Modify();
+
+		FMeshSectionInfo SectionInfo = StaticMesh->GetSectionInfoMap().Get(LODIndex, SectionIndex);
+
+		SectionInfo.bEnableCollision = bCollisionEnabled;
+
+		StaticMesh->GetSectionInfoMap().Set(LODIndex, SectionIndex, SectionInfo);
+
+		StaticMesh->PostEditChange();
 	}
+}
 
-	if (LODIndex >= StaticMesh->GetNumLODs())
+bool UStaticMeshEditorSubsystem::IsSectionCollisionEnabled(const UStaticMesh* StaticMesh, int32 LODIndex, int32 SectionIndex)
+{
+	TGuardValue<bool> UnattendedScriptGuard(GIsRunningUnattendedScript, true);
+
+	if (CanModifyStaticMeshSection(StaticMesh, LODIndex, SectionIndex, "IsSectionCollisionEnabled"))
 	{
-		UE_LOG(LogStaticMeshEditorSubsystem, Error, TEXT("IsSectionCollisionEnabled: Invalid LOD index %d (of %d)."), LODIndex, StaticMesh->GetNumLODs());
-		return false;
+		FMeshSectionInfo SectionInfo = StaticMesh->GetSectionInfoMap().Get(LODIndex, SectionIndex);
+		return SectionInfo.bEnableCollision;
 	}
+	return false;
+}
 
-	if (SectionIndex >= StaticMesh->GetNumSections(LODIndex))
+void UStaticMeshEditorSubsystem::EnableSectionVisibleInRayTracing(UStaticMesh* StaticMesh, bool bVisibleInRayTracing, int32 LODIndex, int32 SectionIndex)
+{
+	TGuardValue<bool> UnattendedScriptGuard(GIsRunningUnattendedScript, true);
+
+	if (CanModifyStaticMeshSection(StaticMesh, LODIndex, SectionIndex, "EnableSectionVisibleInRayTracing"))
 	{
-		UE_LOG(LogStaticMeshEditorSubsystem, Error, TEXT("IsSectionCollisionEnabled: Invalid section index %d (of %d)."), SectionIndex, StaticMesh->GetNumSections(LODIndex));
-		return false;
-	}
+		StaticMesh->Modify();
 
-	FMeshSectionInfo SectionInfo = StaticMesh->GetSectionInfoMap().Get(LODIndex, SectionIndex);
-	return SectionInfo.bEnableCollision;
+		FMeshSectionInfo SectionInfo = StaticMesh->GetSectionInfoMap().Get(LODIndex, SectionIndex);
+
+		SectionInfo.bVisibleInRayTracing = bVisibleInRayTracing;
+
+		StaticMesh->GetSectionInfoMap().Set(LODIndex, SectionIndex, SectionInfo);
+
+		StaticMesh->PostEditChange();
+	}
 }
 
-void UStaticMeshEditorSubsystem::EnableSectionCastShadow(UStaticMesh* StaticMesh, bool bCastShadow, int32 LODIndex, int32 SectionIndex)
+bool UStaticMeshEditorSubsystem::IsSectionVisibleInRayTracingEnabled(const UStaticMesh* StaticMesh, int32 LODIndex, int32 SectionIndex)
 {
 	TGuardValue<bool> UnattendedScriptGuard(GIsRunningUnattendedScript, true);
 
-	if (!EditorScriptingHelpers::CheckIfInEditorAndPIE())
+	if (CanModifyStaticMeshSection(StaticMesh, LODIndex, SectionIndex, "IsSectionVisibleInRayTracingEnabled"))
 	{
-		return;
+		FMeshSectionInfo SectionInfo = StaticMesh->GetSectionInfoMap().Get(LODIndex, SectionIndex);
+		return SectionInfo.bVisibleInRayTracing;
 	}
+	return false;
+}
 
-	if (StaticMesh == nullptr)
+void UStaticMeshEditorSubsystem::EnableSectionAffectDistanceFieldLighting(UStaticMesh* StaticMesh, bool bAffectDistanceFieldLighting, int32 LODIndex, int32 SectionIndex)
+{
+	TGuardValue<bool> UnattendedScriptGuard(GIsRunningUnattendedScript, true);
+
+	if (CanModifyStaticMeshSection(StaticMesh, LODIndex, SectionIndex, "EnableSectionAffectDistanceFieldLighting"))
 	{
-		UE_LOG(LogStaticMeshEditorSubsystem, Error, TEXT("EnableSectionCastShadow: The StaticMesh is null."));
-		return;
+		StaticMesh->Modify();
+
+		FMeshSectionInfo SectionInfo = StaticMesh->GetSectionInfoMap().Get(LODIndex, SectionIndex);
+
+		SectionInfo.bAffectDistanceFieldLighting = bAffectDistanceFieldLighting;
+
+		StaticMesh->GetSectionInfoMap().Set(LODIndex, SectionIndex, SectionInfo);
+
+		StaticMesh->PostEditChange();
 	}
+}
 
-	if (LODIndex >= StaticMesh->GetNumLODs())
+bool UStaticMeshEditorSubsystem::IsSectionAffectDistanceFieldLightingEnabled(const UStaticMesh* StaticMesh, int32 LODIndex, int32 SectionIndex)
+{
+	TGuardValue<bool> UnattendedScriptGuard(GIsRunningUnattendedScript, true);
+
+	if (CanModifyStaticMeshSection(StaticMesh, LODIndex, SectionIndex, "IsSectionAffectDistanceFieldLightingEnabled"))
 	{
-		UE_LOG(LogStaticMeshEditorSubsystem, Error, TEXT("EnableSectionCastShadow: Invalid LOD index %d (of %d)."), LODIndex, StaticMesh->GetNumLODs());
-		return;
+		FMeshSectionInfo SectionInfo = StaticMesh->GetSectionInfoMap().Get(LODIndex, SectionIndex);
+		return SectionInfo.bAffectDistanceFieldLighting;
 	}
+	return false;
+}
 
-	if (SectionIndex >= StaticMesh->GetNumSections(LODIndex))
+void UStaticMeshEditorSubsystem::EnableSectionForceOpaque(UStaticMesh* StaticMesh, bool bForceOpaque, int32 LODIndex, int32 SectionIndex)
+{
+	TGuardValue<bool> UnattendedScriptGuard(GIsRunningUnattendedScript, true);
+
+	if (CanModifyStaticMeshSection(StaticMesh, LODIndex, SectionIndex, "EnableSectionForceOpaque"))
 	{
-		UE_LOG(LogStaticMeshEditorSubsystem, Error, TEXT("EnableSectionCastShadow: Invalid section index %d (of %d)."), SectionIndex, StaticMesh->GetNumSections(LODIndex));
-		return;
-	}
+		StaticMesh->Modify();
 
-	StaticMesh->Modify();
+		FMeshSectionInfo SectionInfo = StaticMesh->GetSectionInfoMap().Get(LODIndex, SectionIndex);
 
-	FMeshSectionInfo SectionInfo = StaticMesh->GetSectionInfoMap().Get(LODIndex, SectionIndex);
+		SectionInfo.bForceOpaque = bForceOpaque;
 
-	SectionInfo.bCastShadow = bCastShadow;
+		StaticMesh->GetSectionInfoMap().Set(LODIndex, SectionIndex, SectionInfo);
 
-	StaticMesh->GetSectionInfoMap().Set(LODIndex, SectionIndex, SectionInfo);
+		StaticMesh->PostEditChange();
+	}
+}
 
-	StaticMesh->PostEditChange();
+bool UStaticMeshEditorSubsystem::IsSectionForceOpaqueEnabled(const UStaticMesh* StaticMesh, int32 LODIndex, int32 SectionIndex)
+{
+	TGuardValue<bool> UnattendedScriptGuard(GIsRunningUnattendedScript, true);
+
+	if (CanModifyStaticMeshSection(StaticMesh, LODIndex, SectionIndex, "IsSectionForceOpaqueEnabled"))
+	{
+		FMeshSectionInfo SectionInfo = StaticMesh->GetSectionInfoMap().Get(LODIndex, SectionIndex);
+		return SectionInfo.bForceOpaque;
+	}
+	return false;
 }
 
 void UStaticMeshEditorSubsystem::SetLODMaterialSlot(UStaticMesh* StaticMesh, int32 MaterialSlotIndex, int32 LODIndex, int32 SectionIndex)
@@ -1658,7 +1730,9 @@ void UStaticMeshEditorSubsystem::SetLODMaterialSlot(UStaticMesh* StaticMesh, int
 	StaticMesh->PostEditChange();
 }
 
-int32 UStaticMeshEditorSubsystem::GetLODMaterialSlot( UStaticMesh* StaticMesh, int32 LODIndex, int32 SectionIndex )
+
+
+int32 UStaticMeshEditorSubsystem::GetLODMaterialSlot(const UStaticMesh* StaticMesh, int32 LODIndex, int32 SectionIndex )
 {
 	TGuardValue<bool> UnattendedScriptGuard( GIsRunningUnattendedScript, true );
 
@@ -1688,7 +1762,7 @@ int32 UStaticMeshEditorSubsystem::GetLODMaterialSlot( UStaticMesh* StaticMesh, i
 	return StaticMesh->GetSectionInfoMap().Get( LODIndex, SectionIndex ).MaterialIndex;
 }
 
-bool UStaticMeshEditorSubsystem::HasVertexColors(UStaticMesh* StaticMesh)
+bool UStaticMeshEditorSubsystem::HasVertexColors(const UStaticMesh* StaticMesh)
 {
 	TGuardValue<bool> UnattendedScriptGuard(GIsRunningUnattendedScript, true);
 
@@ -1725,7 +1799,7 @@ bool UStaticMeshEditorSubsystem::HasVertexColors(UStaticMesh* StaticMesh)
 	return false;
 }
 
-bool UStaticMeshEditorSubsystem::HasInstanceVertexColors(UStaticMeshComponent* StaticMeshComponent)
+bool UStaticMeshEditorSubsystem::HasInstanceVertexColors(const UStaticMeshComponent* StaticMeshComponent)
 {
 	TGuardValue<bool> UnattendedScriptGuard(GIsRunningUnattendedScript, true);
 
@@ -1798,7 +1872,7 @@ bool UStaticMeshEditorSubsystem::SetGenerateLightmapUVs(UStaticMesh* StaticMesh,
 	return false;
 }
 
-int32 UStaticMeshEditorSubsystem::GetNumberVerts(UStaticMesh* StaticMesh, int32 LODIndex)
+int32 UStaticMeshEditorSubsystem::GetNumberVerts(const UStaticMesh* StaticMesh, int32 LODIndex)
 {
 	TGuardValue<bool> UnattendedScriptGuard(GIsRunningUnattendedScript, true);
 
@@ -1816,7 +1890,7 @@ int32 UStaticMeshEditorSubsystem::GetNumberVerts(UStaticMesh* StaticMesh, int32
 	return StaticMesh->GetNumVertices(LODIndex);
 }
 
-int32 UStaticMeshEditorSubsystem::GetNumberMaterials(UStaticMesh* StaticMesh)
+int32 UStaticMeshEditorSubsystem::GetNumberMaterials(const UStaticMesh* StaticMesh)
 {
 	TGuardValue<bool> UnattendedScriptGuard(GIsRunningUnattendedScript, true);
 
@@ -1854,7 +1928,7 @@ void UStaticMeshEditorSubsystem::SetAllowCPUAccess(UStaticMesh* StaticMesh, bool
 	StaticMesh->PostEditChange();
 }
 
-int32 UStaticMeshEditorSubsystem::GetNumUVChannels(UStaticMesh* StaticMesh, int32 LODIndex)
+int32 UStaticMeshEditorSubsystem::GetNumUVChannels(const UStaticMesh* StaticMesh, int32 LODIndex)
 {
 	TGuardValue<bool> UnattendedScriptGuard(GIsRunningUnattendedScript, true);
 
diff --git a/Engine/Source/Editor/StaticMeshEditor/Public/StaticMeshEditorSubsystem.h b/Engine/Source/Editor/StaticMeshEditor/Public/StaticMeshEditorSubsystem.h
--- a/Engine/Source/Editor/StaticMeshEditor/Public/StaticMeshEditorSubsystem.h
+++ b/Engine/Source/Editor/StaticMeshEditor/Public/StaticMeshEditorSubsystem.h
@@ -95,7 +95,7 @@ public:
 	 * @return LODGroup
 	 */
 	UFUNCTION(BlueprintCallable,Category = "Editor Scripting | StaticMesh")
-	FName GetLODGroup(UStaticMesh* StaticMesh);
+	FName GetLODGroup(const UStaticMesh* StaticMesh);
 
 	/**
 	 * Set the LODGroup for the specified static mesh
@@ -149,7 +149,7 @@ public:
 	 * An negative value indicates that the command could not be executed. See log for explanation.
 	 */
 	UFUNCTION(BlueprintCallable, Category = "Static Mesh Utilities")
-	int32 GetLodCount(UStaticMesh* StaticMesh);
+	int32 GetLodCount(const UStaticMesh* StaticMesh);
 
 	/**
 	 * Remove LODs on a static mesh except LOD 0.
@@ -165,7 +165,7 @@ public:
 	 * @return array of LOD screen sizes.
 	 */
 	UFUNCTION(BlueprintPure, Category = "Static Mesh Utilities")
-	TArray<float> GetLodScreenSizes(UStaticMesh* StaticMesh);
+	TArray<float> GetLodScreenSizes(const UStaticMesh* StaticMesh);
 
 	/**
 	 * Set LOD screen sizes.
@@ -183,7 +183,7 @@ public:
 	* @return FMeshNaniteSettings struct for the given static mesh
 	*/
 	UFUNCTION(BlueprintPure, Category = "Static Mesh Utilities")
-	FMeshNaniteSettings GetNaniteSettings(UStaticMesh* StaticMesh);
+	FMeshNaniteSettings GetNaniteSettings(const UStaticMesh* StaticMesh);
 
 	/**
 	* Get the Nanite Settings for the mesh
@@ -223,7 +223,7 @@ public:
 	 * An negative value indicates that the command could not be executed. See log for explanation.
 	 */
 	UFUNCTION(BlueprintCallable, Category = "Static Mesh Utilities")
-	int32 GetSimpleCollisionCount(UStaticMesh* StaticMesh);
+	int32 GetSimpleCollisionCount(const UStaticMesh* StaticMesh);
 
 	/**
 	 * Get the Collision Trace behavior of a static mesh
@@ -231,7 +231,7 @@ public:
 	 * @return the Collision Trace behavior.
 	 */
 	UFUNCTION(BlueprintPure, Category = "Static Mesh Utilities")
-	TEnumAsByte<ECollisionTraceFlag> GetCollisionComplexity(UStaticMesh* StaticMesh);
+	TEnumAsByte<ECollisionTraceFlag> GetCollisionComplexity(const UStaticMesh* StaticMesh);
 
 	/**
 	 * Get number of convex collisions present on a static mesh.
@@ -240,7 +240,7 @@ public:
 	 * An negative value indicates that the command could not be executed. See log for explanation.
 	 */
 	UFUNCTION(BlueprintCallable, Category = "Static Mesh Utilities")
-	int32 GetConvexCollisionCount(UStaticMesh* StaticMesh);
+	int32 GetConvexCollisionCount(const UStaticMesh* StaticMesh);
 
 	/**
 	 * Add a convex collision to a static mesh.
@@ -307,6 +307,32 @@ public:
 		return RemoveCollisionsWithNotification(StaticMesh, true);
 	}
 
+private:
+	/**
+	* Checks if section and lod indices are valid indices.
+	*/
+	static bool CanModifyStaticMeshSection(const UStaticMesh* StaticMesh, int32 LODIndex, int32 SectionIndex, const FString& LogFunctionName);
+
+public:
+	/**
+	 * Enables/disables mesh section shadow casting for a specific LOD.
+	 * @param	StaticMesh			Static mesh to enable/disable section's shadow casting from.
+	 * @param	bCastShadow			If the section should cast shadow.
+	 * @param	LODIndex			Index of the StaticMesh LOD.
+	 * @param	SectionIndex		Index of the StaticMesh Section.
+	 */
+	UFUNCTION(BlueprintCallable, Category = "Static Mesh Utilities")
+	void EnableSectionCastShadow(UStaticMesh* StaticMesh, bool bCastShadow, int32 LODIndex, int32 SectionIndex);
+
+	/**
+	 * Checks if a specific LOD mesh section can cast shadows.
+	 * @param	StaticMesh			Static mesh to check if its section is casting shadows.
+	 * @param	LODIndex			Index of the StaticMesh LOD.
+	 * @param	SectionIndex		Index of the StaticMesh Section.
+	 */
+	UFUNCTION(BlueprintCallable, Category = "Static Mesh Utilities")
+	bool IsSectionCastShadowEnabled(const UStaticMesh* StaticMesh, int32 LODIndex, int32 SectionIndex);
+
 	/**
 	 * Enables/disables mesh section collision for a specific LOD.
 	 * @param	StaticMesh			Static mesh to Enables/disables collisions from.
@@ -319,26 +345,73 @@ public:
 
 	/**
 	 * Checks if a specific LOD mesh section has collision.
-	 * @param	StaticMesh			Static mesh to remove collisions from.
+	 * @param	StaticMesh			Static mesh to check section collisions from.
 	 * @param	LODIndex			Index of the StaticMesh LOD.
 	 * @param	SectionIndex		Index of the StaticMesh Section.
 	 * @return True if the collision is enabled for the specified LOD of the StaticMesh section.
 	 */
 	UFUNCTION(BlueprintPure, Category = "Static Mesh Utilities")
-	bool IsSectionCollisionEnabled(UStaticMesh* StaticMesh, int32 LODIndex, int32 SectionIndex);
+	bool IsSectionCollisionEnabled(const UStaticMesh* StaticMesh, int32 LODIndex, int32 SectionIndex);
 
-public:
+	/**
+	 * Enables/disables mesh section's visibility in ray tracing for a specific LOD.
+	 * @param	StaticMesh				Static mesh to enable/disable section's visibility in ray tracing from.
+	 * @param	bVisibleInRayTracing	If the section visible in ray tracing.
+	 * @param	LODIndex				Index of the StaticMesh LOD.
+	 * @param	SectionIndex			Index of the StaticMesh Section.
+	 */
+	UFUNCTION(BlueprintCallable, Category = "Static Mesh Utilities")
+	void EnableSectionVisibleInRayTracing(UStaticMesh* StaticMesh, bool bVisibleInRayTracing, int32 LODIndex, int32 SectionIndex);
 
 	/**
-	 * Enables/disables mesh section shadow casting for a specific LOD.
-	 * @param	StaticMesh			Static mesh to Enables/disables shadow casting from.
-	 * @param	bCastShadow			If the section should cast shadow.
+	 * Checks if a specific LOD mesh section is visible in ray tracing.
+	 * @param	StaticMesh			Static mesh to check if its section is visible in ray tracing.
 	 * @param	LODIndex			Index of the StaticMesh LOD.
 	 * @param	SectionIndex		Index of the StaticMesh Section.
 	 */
 	UFUNCTION(BlueprintCallable, Category = "Static Mesh Utilities")
-	void EnableSectionCastShadow(UStaticMesh* StaticMesh, bool bCastShadow, int32 LODIndex, int32 SectionIndex);
+	bool IsSectionVisibleInRayTracingEnabled(const UStaticMesh* StaticMesh, int32 LODIndex, int32 SectionIndex);
 
+	/**
+	 * Enables/disables mesh section affects distance field lighting for a specific LOD.
+	 * @param	StaticMesh						Static mesh to enable/disable section affecting distance field lighting from.
+	 * @param	bAffectDistanceFieldLighting	If the section affects distance field lighting.
+	 * @param	LODIndex						Index of the StaticMesh LOD.
+	 * @param	SectionIndex					Index of the StaticMesh Section.
+	 */
+	UFUNCTION(BlueprintCallable, Category = "Static Mesh Utilities")
+	void EnableSectionAffectDistanceFieldLighting(UStaticMesh* StaticMesh, bool bAffectDistanceFieldLighting, int32 LODIndex, int32 SectionIndex);
+
+	/**
+	 * Checks if a specific LOD mesh section affects distance field lighting.
+	 * @param	StaticMesh			Static mesh to check if its section affects distance field lighting.
+	 * @param	LODIndex			Index of the StaticMesh LOD.
+	 * @param	SectionIndex		Index of the StaticMesh Section.
+	 */
+	UFUNCTION(BlueprintCallable, Category = "Static Mesh Utilities")
+	bool IsSectionAffectDistanceFieldLightingEnabled(const UStaticMesh* StaticMesh, int32 LODIndex, int32 SectionIndex);
+
+
+	/**
+	 * Set/unset mesh section to be forcefully opaque for a specific LOD.
+	 * @param	StaticMesh			Static mesh to make section forcefully opaque from.
+	 * @param	bForceOpaque		If the section is forced to be opaque.
+	 * @param	LODIndex			Index of the StaticMesh LOD.
+	 * @param	SectionIndex		Index of the StaticMesh Section.
+	 */
+	UFUNCTION(BlueprintCallable, Category = "Static Mesh Utilities")
+	void EnableSectionForceOpaque(UStaticMesh* StaticMesh, bool bForceOpaque, int32 LODIndex, int32 SectionIndex);
+
+	/**
+	 * Checks if a specific LOD mesh section is forced to be opaque.
+	 * @param	StaticMesh			Static mesh to is forced to be opaque.
+	 * @param	LODIndex			Index of the StaticMesh LOD.
+	 * @param	SectionIndex		Index of the StaticMesh Section.
+	 */
+	UFUNCTION(BlueprintCallable, Category = "Static Mesh Utilities")
+	bool IsSectionForceOpaqueEnabled(const UStaticMesh* StaticMesh, int32 LODIndex, int32 SectionIndex);
+
+public:
 	/**
 	* Sets the material slot for a specific LOD.
 	* @param	StaticMesh			Static mesh to Enables/disables shadow casting from.
@@ -357,15 +430,15 @@ public:
 	 * @return  MaterialSlotIndex	Index of the material slot used by the section or INDEX_NONE in case of error.
 	 */
 	UFUNCTION(BlueprintCallable, Category = "Static Mesh Utilities")
-	int32 GetLODMaterialSlot(UStaticMesh* StaticMesh, int32 LODIndex, int32 SectionIndex);
+	int32 GetLODMaterialSlot(const UStaticMesh* StaticMesh, int32 LODIndex, int32 SectionIndex);
 
 	/** Check whether a static mesh has vertex colors */
 	UFUNCTION(BlueprintPure, Category = "Static Mesh Utilities")
-	bool HasVertexColors(UStaticMesh* StaticMesh);
+	bool HasVertexColors(const UStaticMesh* StaticMesh);
 
 	/** Check whether a static mesh component has vertex colors */
 	UFUNCTION(BlueprintPure, Category = "Static Mesh Utilities")
-	bool HasInstanceVertexColors(UStaticMeshComponent* StaticMeshComponent);
+	bool HasInstanceVertexColors(const UStaticMeshComponent* StaticMeshComponent);
 
 	/** Set Generate Lightmap UVs for StaticMesh */
 	UFUNCTION(BlueprintCallable, Category = "Static Mesh Utilities", meta = (ScriptName = "SetGenerateLightmapUv"))
@@ -373,11 +446,11 @@ public:
 
 	/** Get number of StaticMesh verts for an LOD */
 	UFUNCTION(BlueprintPure, Category = "Static Mesh Utilities")
-	int32 GetNumberVerts(UStaticMesh* StaticMesh, int32 LODIndex);
+	int32 GetNumberVerts(const UStaticMesh* StaticMesh, int32 LODIndex);
 
 	/** Get number of Materials for a StaticMesh */
 	UFUNCTION(BlueprintPure, Category = "Static Mesh Utilities")
-	int32 GetNumberMaterials(UStaticMesh* StaticMesh);
+	int32 GetNumberMaterials(const UStaticMesh* StaticMesh);
 
 	/** Sets StaticMeshFlag bAllowCPUAccess  */
 	UFUNCTION(BlueprintCallable, Category = "Static Mesh Utilities")
@@ -392,7 +465,7 @@ public:
 	 * @return the number of UV channels.
 	 */
 	UFUNCTION(BlueprintCallable, Category = "Static Mesh Utilities")
-	int32 GetNumUVChannels(UStaticMesh* StaticMesh, int32 LODIndex);
+	int32 GetNumUVChannels(const UStaticMesh* StaticMesh, int32 LODIndex);
 
 	/**
 	 * Adds an empty UV channel at the end of the existing channels on the given LOD of a StaticMesh.
diff --git a/Engine/Source/Runtime/Engine/Classes/Engine/StaticMesh.h b/Engine/Source/Runtime/Engine/Classes/Engine/StaticMesh.h
--- a/Engine/Source/Runtime/Engine/Classes/Engine/StaticMesh.h
+++ b/Engine/Source/Runtime/Engine/Classes/Engine/StaticMesh.h
@@ -1582,7 +1582,7 @@ public:
 	 * @param	LODIndex			Index of the StaticMesh LOD.
 	 * @return the number of UV channels.
 	 */
-	ENGINE_API int32 GetNumUVChannels(int32 LODIndex);
+	ENGINE_API int32 GetNumUVChannels(int32 LODIndex) const;
 
 	/** Pre-build navigation collision */
 private:
@@ -1802,6 +1802,7 @@ public:
 	/**
 	 * Returns the number of vertices for the specified LOD.
 	 */
+	UFUNCTION(BlueprintPure, Category = StaticMesh)
 	ENGINE_API int32 GetNumVertices(int32 LODIndex) const;
 
 	/**
diff --git a/Engine/Source/Runtime/Engine/Private/StaticMesh.cpp b/Engine/Source/Runtime/Engine/Private/StaticMesh.cpp
--- a/Engine/Source/Runtime/Engine/Private/StaticMesh.cpp
+++ b/Engine/Source/Runtime/Engine/Private/StaticMesh.cpp
@@ -6104,7 +6104,7 @@ bool UStaticMesh::SetUVChannel(int32 LODIndex, int32 UVChannelIndex, const TMap<
 
 #endif
 
-int32 UStaticMesh::GetNumUVChannels(int32 LODIndex)
+int32 UStaticMesh::GetNumUVChannels(int32 LODIndex) const
 {
 	int32 NumUVChannels = 0;
 #if WITH_EDITORONLY_DATA
