Fix a crash with a unchecked pointer in staticmesh draw canvas when displaying the uvs.

#rnx

diff --git a/Engine/Source/Editor/StaticMeshEditor/Private/StaticMeshEditorViewportClient.cpp b/Engine/Source/Editor/StaticMeshEditor/Private/StaticMeshEditorViewportClient.cpp
--- a/Engine/Source/Editor/StaticMeshEditor/Private/StaticMeshEditorViewportClient.cpp
+++ b/Engine/Source/Editor/StaticMeshEditor/Private/StaticMeshEditorViewportClient.cpp
@@ -750,7 +750,7 @@ void FStaticMeshEditorViewportClient::DrawCanvas(FViewport& InViewport, FSceneVi
 {
 	auto StaticMeshEditor = StaticMeshEditorPtr.Pin();
 	auto StaticMeshEditorViewport = StaticMeshEditorViewportPtr.Pin();
-	if (!StaticMeshEditor.IsValid() || !StaticMeshEditorViewport.IsValid())
+	if (!StaticMeshEditor.IsValid() || !StaticMeshEditorViewport.IsValid() || !StaticMesh || StaticMesh->IsCompiling())
 	{
 		return;
 	}
@@ -882,51 +882,48 @@ void FStaticMeshEditorViewportClient::DrawCanvas(FViewport& InViewport, FSceneVi
 	const FText StaticMeshTriangleCount = FText::AsNumber(StaticMeshEditorPtr.Pin()->GetNumTriangles(CurrentLODLevel));
 	const FText StaticMeshVertexCount = FText::AsNumber(StaticMeshEditorPtr.Pin()->GetNumVertices(CurrentLODLevel));
 
-	if (StaticMesh->IsNaniteEnabled())
+	if (StaticMesh->IsNaniteEnabled() && StaticMesh->GetRenderData())
 	{
-		if (StaticMesh->GetRenderData())
+		const Nanite::FResources& Resources = *StaticMesh->GetRenderData()->NaniteResourcesPtr.Get();
+		if (Resources.RootData.Num() > 0)
 		{
-			const Nanite::FResources& Resources = *StaticMesh->GetRenderData()->NaniteResourcesPtr.Get();
-			if (Resources.RootData.Num() > 0)
-			{
-				// Nanite Mesh Information
-				const FText NaniteTriangleCount = FText::AsNumber(Resources.NumInputTriangles);
-				const FText NaniteVertexCount = FText::AsNumber(Resources.NumInputVertices);
+			// Nanite Mesh Information
+			const FText NaniteTriangleCount = FText::AsNumber(Resources.NumInputTriangles);
+			const FText NaniteVertexCount = FText::AsNumber(Resources.NumInputVertices);
 
-				TextItems.Emplace(FText::Format(NSLOCTEXT("UnrealEd", "Nanite_Triangles_F", "Nanite Triangles:  {0}"), NaniteTriangleCount));
-				TextItems.Emplace(FText::Format(NSLOCTEXT("UnrealEd", "Nanite_Vertices_F", "Nanite Vertices:  {0}"), NaniteVertexCount));
+			TextItems.Emplace(FText::Format(NSLOCTEXT("UnrealEd", "Nanite_Triangles_F", "Nanite Triangles:  {0}"), NaniteTriangleCount));
+			TextItems.Emplace(FText::Format(NSLOCTEXT("UnrealEd", "Nanite_Vertices_F", "Nanite Vertices:  {0}"), NaniteVertexCount));
 
-				// Fallback Mesh Information
-				TextItems.Emplace(FText::Format(NSLOCTEXT("UnrealEd", "Fallback_Triangles_F", "Fallback Triangles:  {0}"), StaticMeshTriangleCount));
-				TextItems.Emplace(FText::Format(NSLOCTEXT("UnrealEd", "Fallback_Vertices_F", "Fallback Vertices:  {0}"), StaticMeshVertexCount));
+			// Fallback Mesh Information
+			TextItems.Emplace(FText::Format(NSLOCTEXT("UnrealEd", "Fallback_Triangles_F", "Fallback Triangles:  {0}"), StaticMeshTriangleCount));
+			TextItems.Emplace(FText::Format(NSLOCTEXT("UnrealEd", "Fallback_Vertices_F", "Fallback Vertices:  {0}"), StaticMeshVertexCount));
 
-				// Ray Tracing Fallback Mesh Information
-				const FStaticMeshRayTracingProxy* RayTracingProxy = StaticMesh->GetRenderData()->RayTracingProxy;
-				if (RayTracingProxy && !RayTracingProxy->bUsingRenderingLODs && RayTracingProxy->LODs.Num() > 0)
-				{
-					const int32 PreviewLODIndex = FMath::Clamp(RayTracingProxy->PreviewLODLevel, 0, RayTracingProxy->LODs.Num() - 1);
-					const FStaticMeshRayTracingProxyLOD& PreviewLOD = RayTracingProxy->LODs[PreviewLODIndex];
-					const FStaticMeshVertexBuffers* VertexBuffers = RayTracingProxy->LODs[PreviewLODIndex].VertexBuffers;
-					const FRawStaticIndexBuffer* IndexBuffer = RayTracingProxy->LODs[PreviewLODIndex].IndexBuffer;
+			// Ray Tracing Fallback Mesh Information
+			const FStaticMeshRayTracingProxy* RayTracingProxy = StaticMesh->GetRenderData()->RayTracingProxy;
+			if (RayTracingProxy && !RayTracingProxy->bUsingRenderingLODs && RayTracingProxy->LODs.Num() > 0)
+			{
+				const int32 PreviewLODIndex = FMath::Clamp(RayTracingProxy->PreviewLODLevel, 0, RayTracingProxy->LODs.Num() - 1);
+				const FStaticMeshRayTracingProxyLOD& PreviewLOD = RayTracingProxy->LODs[PreviewLODIndex];
+				const FStaticMeshVertexBuffers* VertexBuffers = RayTracingProxy->LODs[PreviewLODIndex].VertexBuffers;
+				const FRawStaticIndexBuffer* IndexBuffer = RayTracingProxy->LODs[PreviewLODIndex].IndexBuffer;
 
-					uint32 RayTracingFallbackTriangleCount = 0;
-					uint32 RayTracingFallbackVertexCount = 0;
+				uint32 RayTracingFallbackTriangleCount = 0;
+				uint32 RayTracingFallbackVertexCount = 0;
 
-					if (VertexBuffers && IndexBuffer)
-					{
-						RayTracingFallbackTriangleCount = IndexBuffer->GetNumIndices() / 3;
-						RayTracingFallbackVertexCount = VertexBuffers->PositionVertexBuffer.GetNumVertices();
-					}
+				if (VertexBuffers && IndexBuffer)
+				{
+					RayTracingFallbackTriangleCount = IndexBuffer->GetNumIndices() / 3;
+					RayTracingFallbackVertexCount = VertexBuffers->PositionVertexBuffer.GetNumVertices();
+				}
 
-					if (RayTracingFallbackTriangleCount > 0 && RayTracingFallbackVertexCount > 0)
-					{
-						TextItems.Emplace(FText::Format(
-							NSLOCTEXT("UnrealEd", "RayTracingFallback_Triangles_F", "Ray Tracing Fallback Triangles:  {0}"),
-							FText::AsNumber(RayTracingFallbackTriangleCount)));
-						TextItems.Emplace(FText::Format(
-							NSLOCTEXT("UnrealEd", "RayTracingFallback_Vertices_F", "Ray Tracing Fallback Vertices:  {0}"),
-							FText::AsNumber(RayTracingFallbackVertexCount)));
-					}
+				if (RayTracingFallbackTriangleCount > 0 && RayTracingFallbackVertexCount > 0)
+				{
+					TextItems.Emplace(FText::Format(
+						NSLOCTEXT("UnrealEd", "RayTracingFallback_Triangles_F", "Ray Tracing Fallback Triangles:  {0}"),
+						FText::AsNumber(RayTracingFallbackTriangleCount)));
+					TextItems.Emplace(FText::Format(
+						NSLOCTEXT("UnrealEd", "RayTracingFallback_Vertices_F", "Ray Tracing Fallback Vertices:  {0}"),
+						FText::AsNumber(RayTracingFallbackVertexCount)));
 				}
 			}
 		}
@@ -999,7 +996,7 @@ void FStaticMeshEditorViewportClient::DrawCanvas(FViewport& InViewport, FSceneVi
 		Canvas.DrawShadowedText(static_cast<float>(X), static_cast<float>(Y), Message, GetStatsFont(), FLinearColor(1.0, 0.05, 0.05, 1.0));
 	}
 
-	if (bDrawUVs && StaticMesh->GetRenderData()->LODResources.Num() > 0)
+	if (bDrawUVs && StaticMesh->GetRenderData() && StaticMesh->GetRenderData()->LODResources.Num() > 0)
 	{
 		const int32 YPos = 160;
 		DrawUVsForMesh(Viewport, &Canvas, YPos);
