[Viewport Toolbar] Remove offset from overlay info in Static Mesh Editor

diff --git a/Engine/Source/Editor/StaticMeshEditor/Private/SStaticMeshEditorViewport.cpp b/Engine/Source/Editor/StaticMeshEditor/Private/SStaticMeshEditorViewport.cpp
--- a/Engine/Source/Editor/StaticMeshEditor/Private/SStaticMeshEditorViewport.cpp
+++ b/Engine/Source/Editor/StaticMeshEditor/Private/SStaticMeshEditorViewport.cpp
@@ -321,9 +321,9 @@ void SStaticMeshEditorViewport::SetFloorAndEnvironmentVisibility(bool bVisible)
 
 FMargin SStaticMeshEditorViewport::GetOverlayMargin() const
 {
-	return UE::UnrealEd::ShowNewViewportToolbars() && UE::UnrealEd::ShowOldViewportToolbars()
-			 ? FMargin(6.0f, 72.0f, 6.0f, 6.0f)
-			 : FMargin(6.0f, 36.0f, 6.0f, 6.0f);
+	return UE::UnrealEd::ShowOldViewportToolbars()
+			 ? FMargin(6.0f, 36.0f, 6.0f, 6.0f)
+			 : FMargin(6.0f);
 }
 
 void SStaticMeshEditorViewport::UpdatePreviewSocketMeshes()
