[Viewport Toolbar] Static Mesh Editor - Update preview scene settings menu

diff --git a/Engine/Source/Editor/StaticMeshEditor/Private/SStaticMeshEditorViewport.cpp b/Engine/Source/Editor/StaticMeshEditor/Private/SStaticMeshEditorViewport.cpp
--- a/Engine/Source/Editor/StaticMeshEditor/Private/SStaticMeshEditorViewport.cpp
+++ b/Engine/Source/Editor/StaticMeshEditor/Private/SStaticMeshEditorViewport.cpp
@@ -1028,6 +1028,8 @@ TSharedPtr<SWidget> SStaticMeshEditorViewport::BuildViewportToolbar()
 				UE::UnrealEd::CreateViewportToolbarDefaultContext(SharedThis(this));
 
 			ContextObject->bShowCoordinateSystemControls = false;
+			ContextObject->bShowGridToggleInPreviewSceneSettings = true;
+			ContextObject->ExcludedShowMenuFlags.AddUnique(FEngineShowFlags::SF_Grid);
 
 			ViewportToolbarContext.AddObject(ContextObject);
 		}
diff --git a/Engine/Source/Editor/StaticMeshEditor/Private/StaticMeshEditorViewportToolbarSections.cpp b/Engine/Source/Editor/StaticMeshEditor/Private/StaticMeshEditorViewportToolbarSections.cpp
--- a/Engine/Source/Editor/StaticMeshEditor/Private/StaticMeshEditorViewportToolbarSections.cpp
+++ b/Engine/Source/Editor/StaticMeshEditor/Private/StaticMeshEditorViewportToolbarSections.cpp
@@ -209,6 +209,7 @@ TSharedRef<SWidget> UE::StaticMeshEditor::GenerateShowMenuWidget(const TSharedPt
 			FNewToolMenuDelegate::CreateLambda(
 				[](UToolMenu* InMenu)
 				{
+					constexpr bool bShowGridToggle = true;
 					FillShowSubmenu(InMenu);
 				}
 			)
@@ -259,7 +260,13 @@ void UE::StaticMeshEditor::FillShowSubmenu(UToolMenu* InMenu)
 			MeshComponentsSection.AddSeparator(NAME_None);
 
 			MeshComponentsSection.AddMenuEntry(Commands.SetShowPivot);
-			MeshComponentsSection.AddMenuEntry(Commands.SetShowGrid);
+
+			bool bShowGridToggle = !EditorViewportContext->ExcludedShowMenuFlags.Contains(FEngineShowFlags::SF_Grid);
+			if (bShowGridToggle)
+			{
+				MeshComponentsSection.AddMenuEntry(Commands.SetShowGrid);
+			}
+
 			MeshComponentsSection.AddMenuEntry(Commands.SetShowBounds);
 			MeshComponentsSection.AddMenuEntry(Commands.SetShowSimpleCollision);
 			MeshComponentsSection.AddMenuEntry(Commands.SetShowComplexCollision);
