Avoid clipping Static Mesh LOD menu from viewport toolbar until size is very contrained

Setting the menus clipping prioirity to the same high number range used by other main submenus of the viewport toolbar.

diff --git a/Engine/Source/Editor/StaticMeshEditor/Private/StaticMeshEditorViewportToolbarSections.cpp b/Engine/Source/Editor/StaticMeshEditor/Private/StaticMeshEditorViewportToolbarSections.cpp
--- a/Engine/Source/Editor/StaticMeshEditor/Private/StaticMeshEditorViewportToolbarSections.cpp
+++ b/Engine/Source/Editor/StaticMeshEditor/Private/StaticMeshEditorViewportToolbarSections.cpp
@@ -104,7 +104,7 @@ FToolMenuEntry UE::StaticMeshEditor::CreateLODSubmenu()
 						}
 					);
 
-					InDynamicSection.AddSubMenu(
+					FToolMenuEntry& Entry = InDynamicSection.AddSubMenu(
 						"LOD",
 						Label,
 						FText(),
@@ -115,12 +115,12 @@ FToolMenuEntry UE::StaticMeshEditor::CreateLODSubmenu()
 							}
 						)
 					);
+					Entry.InsertPosition.Position = EToolMenuInsertType::Last;
+					Entry.ToolBarData.ResizeParams.ClippingPriority = 800;
 				}
 			}
 		)
 	);
-	Entry.InsertPosition.Position = EToolMenuInsertType::Last;
-	Entry.ToolBarData.ResizeParams.ClippingPriority = 800;
 	return Entry;
 }
 
