[Viewport Toolbar] Camera Speed options are missing in the Viewport Toolbar of the Static Mesh editor

diff --git a/Engine/Source/Editor/StaticMeshEditor/Private/SStaticMeshEditorViewport.cpp b/Engine/Source/Editor/StaticMeshEditor/Private/SStaticMeshEditorViewport.cpp
--- a/Engine/Source/Editor/StaticMeshEditor/Private/SStaticMeshEditorViewport.cpp
+++ b/Engine/Source/Editor/StaticMeshEditor/Private/SStaticMeshEditorViewport.cpp
@@ -973,6 +973,7 @@ TSharedPtr<SWidget> SStaticMeshEditorViewport::BuildViewportToolbar()
 				// Build the menu name our Camera menu will be using so we can extend it.
 				const FName SubmenuName = UToolMenus::JoinMenuPaths(ViewportToolbarName, "Camera");
 				UE::UnrealEd::ExtendCameraSubmenu(SubmenuName);
+				UE::StaticMeshEditor::ExtendCameraSubmenu(SubmenuName);
 
 				RightSection.AddEntry(UE::UnrealEd::CreateCameraSubmenu());
 			}
diff --git a/Engine/Source/Editor/StaticMeshEditor/Private/StaticMeshEditorViewportToolbarSections.cpp b/Engine/Source/Editor/StaticMeshEditor/Private/StaticMeshEditorViewportToolbarSections.cpp
--- a/Engine/Source/Editor/StaticMeshEditor/Private/StaticMeshEditorViewportToolbarSections.cpp
+++ b/Engine/Source/Editor/StaticMeshEditor/Private/StaticMeshEditorViewportToolbarSections.cpp
@@ -78,6 +78,16 @@ FText UE::StaticMeshEditor::GetLODMenuLabel(const TSharedPtr<SStaticMeshEditorVi
 	return Label;
 }
 
+void UE::StaticMeshEditor::ExtendCameraSubmenu(FName InCameraOptionsSubmenuName)
+{
+	UToolMenu* const Submenu = UToolMenus::Get()->ExtendMenu(InCameraOptionsSubmenuName);
+
+	FToolMenuSection& PositioningSection =
+		Submenu->FindOrAddSection("Positioning", LOCTEXT("PositioningLabel", "Positioning"));
+
+	PositioningSection.AddEntry(UE::UnrealEd::CreateCameraSpeedMenu());
+}
+
 FToolMenuEntry UE::StaticMeshEditor::CreateLODSubmenu()
 {
 	FToolMenuEntry Entry = FToolMenuEntry::InitDynamicEntry(
diff --git a/Engine/Source/Editor/StaticMeshEditor/Private/StaticMeshEditorViewportToolbarSections.h b/Engine/Source/Editor/StaticMeshEditor/Private/StaticMeshEditorViewportToolbarSections.h
--- a/Engine/Source/Editor/StaticMeshEditor/Private/StaticMeshEditorViewportToolbarSections.h
+++ b/Engine/Source/Editor/StaticMeshEditor/Private/StaticMeshEditorViewportToolbarSections.h
@@ -19,4 +19,6 @@ TSharedRef<SWidget> GenerateShowMenuWidget(const TSharedPtr<SStaticMeshEditorVie
 void FillShowSubmenu(UToolMenu *InMenu);
 
 FText GetLODMenuLabel(const TSharedPtr<SStaticMeshEditorViewport>& InStaticMeshEditorViewport);
+
+void ExtendCameraSubmenu(FName InCameraOptionsSubmenuName);
 }
