StaticMeshEditor: Add multi selection for static mesh sockets

#12319

diff --git a/Engine/Source/Editor/StaticMeshEditor/Private/SStaticMeshEditorViewport.cpp b/Engine/Source/Editor/StaticMeshEditor/Private/SStaticMeshEditorViewport.cpp
--- a/Engine/Source/Editor/StaticMeshEditor/Private/SStaticMeshEditorViewport.cpp
+++ b/Engine/Source/Editor/StaticMeshEditor/Private/SStaticMeshEditorViewport.cpp
@@ -911,16 +911,32 @@ void SStaticMeshEditorViewport::BindCommands()
 void SStaticMeshEditorViewport::OnFocusViewportToSelection()
 {
 	// If we have selected sockets, focus on them
-	UStaticMeshSocket* SelectedSocket = StaticMeshEditorPtr.Pin()->GetSelectedSocket();
-	if (SelectedSocket && PreviewMeshComponent)
+	const TArray<UStaticMeshSocket*> SelectedSockets = StaticMeshEditorPtr.Pin()->GetSelectedSockets();
+
+	if (SelectedSockets.Num() && PreviewMeshComponent)
 	{
 		FTransform SocketTransform;
-		SelectedSocket->GetSocketTransform(SocketTransform, PreviewMeshComponent);
-
-		const FVector Extent(30.0f);
+		SelectedSockets[0]->GetSocketTransform(SocketTransform, PreviewMeshComponent);
 
 		const FVector Origin = SocketTransform.GetLocation();
-		const FBox Box(Origin - Extent, Origin + Extent);
+
+		FBox Box(Origin, Origin);
+
+		for (const UStaticMeshSocket* Socket : SelectedSockets)
+		{
+			Socket->GetSocketTransform(SocketTransform, PreviewMeshComponent);
+
+			Box.Max.X = FMath::Max(SocketTransform.GetLocation().X, Box.Max.X);
+			Box.Max.Y = FMath::Max(SocketTransform.GetLocation().Y, Box.Max.Y);
+			Box.Max.Z = FMath::Max(SocketTransform.GetLocation().Z, Box.Max.Z);
+
+			Box.Min.X = FMath::Min(SocketTransform.GetLocation().X, Box.Min.X);
+			Box.Min.Y = FMath::Min(SocketTransform.GetLocation().Y, Box.Min.Y);
+			Box.Min.Z = FMath::Min(SocketTransform.GetLocation().Z, Box.Min.Z);
+		}
+
+		Box.Max += FVector(30.f);
+		Box.Min -= FVector(30.f);
 
 		EditorViewportClient->FocusViewportOnBox(Box);
 		return;
diff --git a/Engine/Source/Editor/StaticMeshEditor/Private/StaticMeshEditor.cpp b/Engine/Source/Editor/StaticMeshEditor/Private/StaticMeshEditor.cpp
--- a/Engine/Source/Editor/StaticMeshEditor/Private/StaticMeshEditor.cpp
+++ b/Engine/Source/Editor/StaticMeshEditor/Private/StaticMeshEditor.cpp
@@ -1091,14 +1091,43 @@ void FStaticMeshEditor::SetSelectedSocket(UStaticMeshSocket* InSelectedSocket)
 	SocketManager->SetSelectedSocket(InSelectedSocket);
 }
 
+void FStaticMeshEditor::AddSelectedSocket(UStaticMeshSocket* InSelectedSocket)
+{
+	SocketManager->AddSelectedSocket(InSelectedSocket);
+}
+
+void FStaticMeshEditor::RemoveSelectedSocket(const UStaticMeshSocket* InSelectedSocket)
+{
+	SocketManager->RemoveSelectedSocket(InSelectedSocket);
+}
+
 UStaticMeshSocket* FStaticMeshEditor::GetSelectedSocket() const
 {
+PRAGMA_DISABLE_DEPRECATION_WARNINGS
 	return SocketManager.IsValid() ? SocketManager->GetSelectedSocket() : nullptr;
+PRAGMA_ENABLE_DEPRECATION_WARNINGS
+}
+
+TArray<UStaticMeshSocket*> FStaticMeshEditor::GetSelectedSockets() const
+{
+	return SocketManager.IsValid() ? SocketManager->GetSelectedSockets() : TArray<UStaticMeshSocket*>();
+}
+
+bool FStaticMeshEditor::HasSelectedSockets() const
+{
+	return SocketManager.IsValid() ? SocketManager->HasSelectedSockets() : false;
 }
 
 void FStaticMeshEditor::DuplicateSelectedSocket()
 {
+PRAGMA_DISABLE_DEPRECATION_WARNINGS
 	SocketManager->DuplicateSelectedSocket();
+PRAGMA_ENABLE_DEPRECATION_WARNINGS
+}
+
+void FStaticMeshEditor::DuplicateSelectedSockets()
+{
+	SocketManager->DuplicateSelectedSockets();
 }
 
 void FStaticMeshEditor::RequestRenameSelectedSocket()
@@ -2450,7 +2479,7 @@ int32 FStaticMeshEditor::GetNumUVChannels( int32 LODLevel ) const
 
 void FStaticMeshEditor::DeleteSelected()
 {
-	if (GetSelectedSocket())
+	if (HasSelectedSockets())
 	{
 		DeleteSelectedSockets();
 	}
@@ -2463,14 +2492,14 @@ void FStaticMeshEditor::DeleteSelected()
 
 bool FStaticMeshEditor::CanDeleteSelected() const
 {
-	return GetOpenMethod() != EAssetOpenMethod::View && (GetSelectedSocket() != NULL || HasSelectedPrims());
+	return GetOpenMethod() != EAssetOpenMethod::View && (HasSelectedSockets() || HasSelectedPrims());
 }
 
 void FStaticMeshEditor::DeleteSelectedSockets()
 {
 	check(SocketManager.IsValid());
 
-	SocketManager->DeleteSelectedSocket();
+	SocketManager->DeleteSelectedSockets();
 }
 
 void FStaticMeshEditor::DeleteSelectedPrims()
@@ -2544,7 +2573,7 @@ void FStaticMeshEditor::DeleteSelectedPrims()
 
 void FStaticMeshEditor::DuplicateSelected()
 {
-	DuplicateSelectedSocket();
+	DuplicateSelectedSockets();
 
 	const FVector InitialOffset(20.f);
 	DuplicateSelectedPrims(&InitialOffset);
@@ -2552,7 +2581,7 @@ void FStaticMeshEditor::DuplicateSelected()
 
 bool FStaticMeshEditor::CanDuplicateSelected() const
 {
-	return GetOpenMethod() != EAssetOpenMethod::View && (GetSelectedSocket() != NULL || HasSelectedPrims());
+	return GetOpenMethod() != EAssetOpenMethod::View && (HasSelectedSockets() || HasSelectedPrims());
 }
 
 void FStaticMeshEditor::CopySelected()
@@ -2585,7 +2614,7 @@ bool FStaticMeshEditor::CanPasteCopied() const
 
 bool FStaticMeshEditor::CanRenameSelected() const
 {
-	return GetOpenMethod() != EAssetOpenMethod::View && (GetSelectedSocket() != NULL);
+	return GetOpenMethod() != EAssetOpenMethod::View && (HasSelectedSockets());
 }
 
 void FStaticMeshEditor::ExecuteFindInExplorer()
@@ -2783,12 +2812,15 @@ void FStaticMeshEditor::OnMeshChanged()
 
 void FStaticMeshEditor::OnSocketSelectionChanged()
 {
-	UStaticMeshSocket* SelectedSocket = GetSelectedSocket();
-	if (SelectedSocket)
+	if (HasSelectedSockets())
 	{
 		ClearSelectedPrims();
+		GetStaticMeshViewport()->GetViewportClient().OnSocketSelectionChanged(GetSelectedSockets()[0]);
+	}
+	else
+	{
+		GetStaticMeshViewport()->GetViewportClient().OnSocketSelectionChanged(nullptr);
 	}
-	GetStaticMeshViewport()->GetViewportClient().OnSocketSelectionChanged( SelectedSocket );
 }
 
 void FStaticMeshEditor::OnPostReimport(UObject* InObject, bool bSuccess)
diff --git a/Engine/Source/Editor/StaticMeshEditor/Private/StaticMeshEditor.h b/Engine/Source/Editor/StaticMeshEditor/Private/StaticMeshEditor.h
--- a/Engine/Source/Editor/StaticMeshEditor/Private/StaticMeshEditor.h
+++ b/Engine/Source/Editor/StaticMeshEditor/Private/StaticMeshEditor.h
@@ -119,8 +119,13 @@ public:
 	virtual UStaticMeshComponent* GetStaticMeshComponent() const override;
 
 	virtual UStaticMeshSocket* GetSelectedSocket() const override;
+	virtual TArray<UStaticMeshSocket*> GetSelectedSockets() const override;
+	virtual bool HasSelectedSockets() const override;
 	virtual void SetSelectedSocket(UStaticMeshSocket* InSelectedSocket) override;
+	virtual void AddSelectedSocket(UStaticMeshSocket* InSelectedSocket) override;
+	virtual void RemoveSelectedSocket(const UStaticMeshSocket* InSelectedSocket) override;
 	virtual void DuplicateSelectedSocket() override;
+	virtual void DuplicateSelectedSockets() override;
 	virtual void RequestRenameSelectedSocket() override;
 
 	virtual bool IsPrimValid(const FPrimData& InPrimData) const override;
diff --git a/Engine/Source/Editor/StaticMeshEditor/Private/StaticMeshEditorViewportClient.cpp b/Engine/Source/Editor/StaticMeshEditor/Private/StaticMeshEditorViewportClient.cpp
--- a/Engine/Source/Editor/StaticMeshEditor/Private/StaticMeshEditorViewportClient.cpp
+++ b/Engine/Source/Editor/StaticMeshEditor/Private/StaticMeshEditorViewportClient.cpp
@@ -173,47 +173,50 @@ bool FStaticMeshEditorViewportClient::InputWidgetDelta( FViewport* InViewport, E
 	{
 		if (CurrentAxis != EAxisList::None)
 		{
-			UStaticMeshSocket* SelectedSocket = StaticMeshEditorPtr.Pin()->GetSelectedSocket();
-			if(SelectedSocket)
+			const TArray<UStaticMeshSocket*> SelectedSockets = StaticMeshEditorPtr.Pin()->GetSelectedSockets();
+			if (SelectedSockets.Num())
 			{
-				FProperty* ChangedProperty = NULL;
-				const UE::Widget::EWidgetMode MoveMode = GetWidgetMode();
-				if(MoveMode == UE::Widget::WM_Rotate)
+				for (UStaticMeshSocket* SelectedSocket : SelectedSockets)
 				{
-					ChangedProperty = FindFProperty<FProperty>( UStaticMeshSocket::StaticClass(), "RelativeRotation" );
-					SelectedSocket->PreEditChange(ChangedProperty);
-
-					FRotator CurrentRot = SelectedSocket->RelativeRotation;
-					FRotator SocketWinding, SocketRotRemainder;
-					CurrentRot.GetWindingAndRemainder(SocketWinding, SocketRotRemainder);
-
-					const FQuat ActorQ = SocketRotRemainder.Quaternion();
-					const FQuat DeltaQ = Rot.Quaternion();
-					const FQuat ResultQ = DeltaQ * ActorQ;
-					const FRotator NewSocketRotRem = FRotator( ResultQ );
-					FRotator DeltaRot = NewSocketRotRem - SocketRotRemainder;
-					DeltaRot.Normalize();
-
-					SelectedSocket->RelativeRotation += DeltaRot;
-					SelectedSocket->RelativeRotation = SelectedSocket->RelativeRotation.Clamp();
-				}
-				else if(MoveMode == UE::Widget::WM_Translate)
-				{
-					ChangedProperty = FindFProperty<FProperty>( UStaticMeshSocket::StaticClass(), "RelativeLocation" );
-					SelectedSocket->PreEditChange(ChangedProperty);
+					FProperty* ChangedProperty = NULL;
+					const UE::Widget::EWidgetMode MoveMode = GetWidgetMode();
+					if (MoveMode == UE::Widget::WM_Rotate)
+					{
+						ChangedProperty = FindFProperty<FProperty>(UStaticMeshSocket::StaticClass(), "RelativeRotation");
+						SelectedSocket->PreEditChange(ChangedProperty);
+
+						const FRotator CurrentRot = SelectedSocket->RelativeRotation;
+						FRotator SocketWinding, SocketRotRemainder;
+						CurrentRot.GetWindingAndRemainder(SocketWinding, SocketRotRemainder);
+
+						const FQuat ActorQ = SocketRotRemainder.Quaternion();
+						const FQuat DeltaQ = Rot.Quaternion();
+						const FQuat ResultQ = DeltaQ * ActorQ;
+						const FRotator NewSocketRotRem = FRotator(ResultQ);
+						FRotator DeltaRot = NewSocketRotRem - SocketRotRemainder;
+						DeltaRot.Normalize();
+
+						SelectedSocket->RelativeRotation += DeltaRot;
+						SelectedSocket->RelativeRotation = SelectedSocket->RelativeRotation.Clamp();
+					}
+					else if (MoveMode == UE::Widget::WM_Translate)
+					{
+						ChangedProperty = FindFProperty<FProperty>(UStaticMeshSocket::StaticClass(), "RelativeLocation");
+						SelectedSocket->PreEditChange(ChangedProperty);
 
-					//FRotationMatrix SocketRotTM( SelectedSocket->RelativeRotation );
-					//FVector SocketMove = SocketRotTM.TransformVector( Drag );
+						//FRotationMatrix SocketRotTM( SelectedSocket->RelativeRotation );
+						//FVector SocketMove = SocketRotTM.TransformVector( Drag );
 
-					SelectedSocket->RelativeLocation += Drag;
-				}
-				if ( ChangedProperty )
-				{			
-					FPropertyChangedEvent PropertyChangedEvent( ChangedProperty );
-					SelectedSocket->PostEditChangeProperty(PropertyChangedEvent);
-				}
+						SelectedSocket->RelativeLocation += Drag;
+					}
+					if (ChangedProperty)
+					{
+						FPropertyChangedEvent PropertyChangedEvent(ChangedProperty);
+						SelectedSocket->PostEditChangeProperty(PropertyChangedEvent);
+					}
 
-				StaticMeshEditorPtr.Pin()->GetStaticMesh()->MarkPackageDirty();
+					StaticMeshEditorPtr.Pin()->GetStaticMesh()->MarkPackageDirty();
+				}
 			}
 			else
 			{
@@ -253,8 +256,8 @@ void FStaticMeshEditorViewportClient::TrackingStarted( const struct FInputEventS
 	if( !bManipulating && bIsDraggingWidget && !bTrackingHandledExternally)
 	{
 		Widget->SetSnapEnabled(true);
-		const UStaticMeshSocket* SelectedSocket = StaticMeshEditorPtr.Pin()->GetSelectedSocket();
-		if (SelectedSocket)
+		const TArray<UStaticMeshSocket*> SelectedSockets = StaticMeshEditorPtr.Pin()->GetSelectedSockets();
+		if (SelectedSockets.Num())
 		{
 			FText TransText;
 			if( GetWidgetMode() == UE::Widget::WM_Rotate )
@@ -269,7 +272,7 @@ void FStaticMeshEditorViewportClient::TrackingStarted( const struct FInputEventS
 					if ( bAltDown )
 					{
 						// Rather than moving/rotating the selected socket, copy it and move the copy instead
-						StaticMeshEditorPtr.Pin()->DuplicateSelectedSocket();
+						StaticMeshEditorPtr.Pin()->DuplicateSelectedSockets();
 					}
 				}
 
@@ -329,7 +332,7 @@ UE::Widget::EWidgetMode FStaticMeshEditorViewportClient::GetWidgetMode() const
 	{
 		return ModeTools->GetWidgetMode();
 	}
-	else if(StaticMeshEditorPtr.Pin()->GetSelectedSocket())
+	else if (StaticMeshEditorPtr.Pin()->HasSelectedSockets())
 	{
 		return WidgetMode;
 	}
@@ -369,11 +372,7 @@ bool FStaticMeshEditorViewportClient::CanSetWidgetMode(UE::Widget::EWidgetMode N
 		}
 		else if (NewMode != UE::Widget::WM_Scale)	// Sockets don't support scaling
 		{
-			const UStaticMeshSocket* SelectedSocket = StaticMeshEditorPtr.Pin()->GetSelectedSocket();
-			if (SelectedSocket)
-			{
-				return true;
-			}
+			return StaticMeshEditorPtr.Pin()->HasSelectedSockets();
 		}
 	}
 	return false;
@@ -383,9 +382,9 @@ bool FStaticMeshEditorViewportClient::CanCycleWidgetMode() const
 {
 	if (!Widget->IsDragging())
 	{
-		const UStaticMeshSocket* SelectedSocket = StaticMeshEditorPtr.Pin()->GetSelectedSocket();
+		const bool bHasSelectedSockets = StaticMeshEditorPtr.Pin()->HasSelectedSockets();
 		const bool bSelectedPrim = StaticMeshEditorPtr.Pin()->HasSelectedPrims();
-		if ((SelectedSocket || bSelectedPrim || IsCustomModeUsingWidget()))
+		if ((bHasSelectedSockets || bSelectedPrim || IsCustomModeUsingWidget()))
 		{
 			return true;
 		}
@@ -410,8 +409,9 @@ FVector FStaticMeshEditorViewportClient::GetWidgetLocation() const
 	{
 		return ModeTools->GetWidgetLocation();
 	}
-	else if (const UStaticMeshSocket* SelectedSocket = StaticMeshEditorPtr.Pin()->GetSelectedSocket())
+	else if (StaticMeshEditorPtr.Pin()->HasSelectedSockets())
 	{
+		const UStaticMeshSocket* SelectedSocket = StaticMeshEditorPtr.Pin()->GetSelectedSockets()[0];
 		FMatrix SocketTM;
 		SelectedSocket->GetSocketMatrix(SocketTM, StaticMeshComponent);
 
@@ -438,12 +438,12 @@ FMatrix FStaticMeshEditorViewportClient::GetWidgetCoordSystem() const
 		return ModeTools->GetCustomInputCoordinateSystem();
 	}
 
-	if(const UStaticMeshSocket* SelectedSocket = StaticMeshEditorPtr.Pin()->GetSelectedSocket())
+	if (StaticMeshEditorPtr.Pin()->HasSelectedSockets())
 	{
 		//FMatrix SocketTM;
 		//SelectedSocket->GetSocketMatrix(SocketTM, StaticMeshComponent);
-
-		return FRotationMatrix( SelectedSocket->RelativeRotation );
+		const UStaticMeshSocket* SelectedSocket = StaticMeshEditorPtr.Pin()->GetSelectedSockets()[0];
+		return FRotationMatrix(SelectedSocket->RelativeRotation);
 	}
 
 	FTransform PrimTransform = FTransform::Identity;
@@ -593,17 +593,20 @@ void FStaticMeshEditorViewportClient::Draw(const FSceneView* View,FPrimitiveDraw
 
 	if( bShowSockets )
 	{
-		const FColor SocketColor = FColor(255, 128, 128);
+		const FColor UnselectedSocketColor = FColor(255, 128, 128);
+		const FColor SelectedSocketColor = FColor(255, 220, 220);
 
 		for(int32 i=0; i < StaticMesh->Sockets.Num(); i++)
 		{
 			UStaticMeshSocket* Socket = StaticMesh->Sockets[i];
 			if(Socket)
 			{
+				const bool bIsSocketSelected = StaticMeshEditor->GetSelectedSockets().Contains(Socket);
+
 				FMatrix SocketTM;
 				Socket->GetSocketMatrix(SocketTM, StaticMeshComponent);
 				PDI->SetHitProxy( new HSMESocketProxy(i) );
-				DrawWireDiamond(PDI, SocketTM, 5.f, SocketColor, SDPG_Foreground);
+				DrawWireDiamond(PDI, SocketTM, 5.f, bIsSocketSelected ? SelectedSocketColor : UnselectedSocketColor, SDPG_Foreground);
 				PDI->SetHitProxy( NULL );
 			}
 		}
@@ -778,8 +781,7 @@ void FStaticMeshEditorViewportClient::DrawCanvas(FViewport& InViewport, FSceneVi
 					FCanvasTextItem TextItem(FVector2D(XPos, YPos), FText::FromString(Socket->SocketName.ToString()), GEngine->GetSmallFont(), FLinearColor(FColor(255, 196, 196)));
 					Canvas.DrawItem(TextItem);
 
-					const UStaticMeshSocket* SelectedSocket = StaticMeshEditor->GetSelectedSocket();
-					if (bManipulating && SelectedSocket == Socket)
+					if (bManipulating && StaticMeshEditor->GetSelectedSockets().Contains(Socket))
 					{
 						//Figure out the text height
 						FTextSizingParameters Parameters(GEngine->GetSmallFont(), 1.0f, 1.0f);
@@ -1075,7 +1077,22 @@ void FStaticMeshEditorViewportClient::ProcessClick(class FSceneView& InView, cla
 
 			if(Socket)
 			{
-				StaticMeshEditorPtr.Pin()->SetSelectedSocket(Socket);
+				if (!bCtrlDown)
+				{
+					StaticMeshEditorPtr.Pin()->SetSelectedSocket(Socket);
+				}
+				else
+				{
+					//Select the socket if it is unselected and vice versa
+					if (StaticMeshEditorPtr.Pin()->GetSelectedSockets().Contains(Socket))
+					{
+						StaticMeshEditorPtr.Pin()->RemoveSelectedSocket(Socket);
+					}
+					else
+					{
+						StaticMeshEditorPtr.Pin()->AddSelectedSocket(Socket);
+					}
+				}
 			}
 
 			ClearSelectedSockets = false;
@@ -1110,9 +1127,9 @@ void FStaticMeshEditorViewportClient::ProcessClick(class FSceneView& InView, cla
 		}
 		else if (IsShowSocketsChecked() && HitProxy->IsA(HSMEVertexProxy::StaticGetType()))
 		{
-			UStaticMeshSocket* Socket = StaticMeshEditorPtr.Pin()->GetSelectedSocket();
+			TArray<UStaticMeshSocket*> Sockets = StaticMeshEditorPtr.Pin()->GetSelectedSockets();
 
-			if (Socket)
+			for (UStaticMeshSocket* Socket : Sockets)
 			{
 				HSMEVertexProxy* VertexProxy = (HSMEVertexProxy*)HitProxy;
 				TSharedPtr<IStaticMeshEditor> StaticMeshEditor = StaticMeshEditorPtr.Pin();
@@ -1356,7 +1373,7 @@ void FStaticMeshEditorViewportClient::ProcessClick(class FSceneView& InView, cla
 		}
 	}
 
-	if (ClearSelectedSockets && StaticMeshEditorPtr.Pin()->GetSelectedSocket())
+	if (ClearSelectedSockets && StaticMeshEditorPtr.Pin()->HasSelectedSockets())
 	{
 		StaticMeshEditorPtr.Pin()->SetSelectedSocket(NULL);
 	}
diff --git a/Engine/Source/Editor/StaticMeshEditor/Public/IStaticMeshEditor.h b/Engine/Source/Editor/StaticMeshEditor/Public/IStaticMeshEditor.h
--- a/Engine/Source/Editor/StaticMeshEditor/Public/IStaticMeshEditor.h
+++ b/Engine/Source/Editor/StaticMeshEditor/Public/IStaticMeshEditor.h
@@ -52,8 +52,15 @@ public:
 	/** Retrieves the static mesh component. */
 	virtual UStaticMeshComponent* GetStaticMeshComponent() const = 0;
 
-	/** Retrieves the currently selected socket from the Socket Manager. */
+	/** Retrieves the first currently selected socket from the Socket Manager. */
+	UE_DEPRECATED(5.6, "GetSelectedSocket is deprecated, use GetSelectedSockets instead.")
 	virtual UStaticMeshSocket* GetSelectedSocket() const = 0;
+	
+	/** Retrieves the currently selected sockets from the Socket Manager. */
+	virtual TArray<UStaticMeshSocket*> GetSelectedSockets() const = 0;
+
+	/** Returns true if at least one socket is selected. */
+	virtual bool HasSelectedSockets() const = 0;
 
 	/**
 	 *	Set the currently selected socket in the Socket Manager.
@@ -62,8 +69,26 @@ public:
 	 */
 	virtual void SetSelectedSocket(UStaticMeshSocket* InSelectedSocket) = 0;
 
-	/** Duplicate the selected socket */
+	/**
+	 *	Adds a socket to the currently selected sockets in the Socket Manager.
+	 *
+	 *	@param	InSelectedSocket		The socket to add to the current selection in the Socket Manager.
+	 */
+	virtual void AddSelectedSocket(UStaticMeshSocket* InSelectedSocket) = 0;
+
+	/**
+	*	Removes a socket form the currently selected sockets in the Socket Manager.
+	*
+	*	@param	InSelectedSocket		The socket to removed from the current selection in the Socket Manager.
+	*/
+	virtual void RemoveSelectedSocket(const UStaticMeshSocket* InSelectedSocket) = 0;
+	
+	/** Duplicate the first selected socket */
+	UE_DEPRECATED(5.6, "DuplicateSelectedSocket is deprecated, use DuplicateSelectedSockets instead.")
 	virtual void DuplicateSelectedSocket() = 0;
+	
+	/** Duplicate the selected sockets */
+	virtual void DuplicateSelectedSockets() = 0;
 
 	/** Requests to rename selected socket */
 	virtual void RequestRenameSelectedSocket() = 0;
diff --git a/Engine/Source/Editor/UnrealEd/Private/SSocketManager.cpp b/Engine/Source/Editor/UnrealEd/Private/SSocketManager.cpp
--- a/Engine/Source/Editor/UnrealEd/Private/SSocketManager.cpp
+++ b/Engine/Source/Editor/UnrealEd/Private/SSocketManager.cpp
@@ -341,7 +341,7 @@ void SSocketManager::Construct(const FArguments& InArgs)
 					[
 						SAssignNew(SocketListView, SListView<TSharedPtr< SocketListItem > >)
 
-						.SelectionMode(ESelectionMode::Single)
+						.SelectionMode(ESelectionMode::Multi)
 
 						.ListItemsSource(&SocketList)
 
@@ -431,6 +431,24 @@ UStaticMeshSocket* SSocketManager::GetSelectedSocket() const
 	return nullptr;
 }
 
+TArray<UStaticMeshSocket*> SSocketManager::GetSelectedSockets() const
+{
+	TArray<UStaticMeshSocket*> Sockets;
+	Sockets.Reserve(SocketListView->GetNumItemsSelected());
+
+	for (TSharedPtr<SocketListItem>& Socket : SocketListView->GetSelectedItems())
+	{
+		Sockets.Add(Socket->Socket);
+	}
+
+	return Sockets;
+}
+
+bool SSocketManager::HasSelectedSockets() const
+{
+	return SocketListView->GetNumItemsSelected() > 0;
+}
+
 EVisibility SSocketManager::GetSelectSocketMessageVisibility() const
 {
 	return SocketListView->GetSelectedItems().Num() > 0 ? EVisibility::Hidden : EVisibility::Visible;
@@ -464,6 +482,38 @@ void SSocketManager::SetSelectedSocket(UStaticMeshSocket* InSelectedSocket)
 	}
 }
 
+
+void SSocketManager::AddSelectedSocket(UStaticMeshSocket* InSelectedSocket)
+{
+	TArray<UStaticMeshSocket*> SelectedSockets = GetSelectedSockets();
+	SelectedSockets.AddUnique(InSelectedSocket);
+	SetSelectedSockets(SelectedSockets);
+}
+
+void SSocketManager::RemoveSelectedSocket(const UStaticMeshSocket* InSelectedSocket)
+{
+	TSharedPtr<SocketListItem> SocketItem = nullptr;
+
+	for (int32 i = 0; i < SocketListView->GetItems().Num(); ++i)
+	{
+		if (SocketListView->GetItems()[i]->Socket == InSelectedSocket)
+		{
+			SocketItem = SocketListView->GetItems()[i];
+			break;
+		}
+	}
+
+	if (SocketItem == nullptr)
+	{
+		return;
+	}
+
+	SocketListView->SetItemSelection(SocketItem, false);
+
+	SocketListView->RequestListRefresh();
+	SocketSelectionChanged(GetSelectedSockets());
+}
+
 TSharedRef< ITableRow > SSocketManager::MakeWidgetFromOption( TSharedPtr<SocketListItem> InItem, const TSharedRef< STableViewBase >& OwnerTable )
 {
 	return SNew( SSocketDisplayItem, OwnerTable )
@@ -549,12 +599,84 @@ void SSocketManager::DuplicateSelectedSocket()
 	}
 }
 
+void SSocketManager::DuplicateSelectedSockets()
+{
+	TSharedPtr<IStaticMeshEditor> StaticMeshEditorPinned = StaticMeshEditorPtr.Pin();
+
+	TArray<UStaticMeshSocket*> SelectedSockets = GetSelectedSockets();
+
+	if (!StaticMeshEditorPinned.IsValid() || SelectedSockets.Num() == 0)
+	{
+		return;
+	}
+
+	const FScopedTransaction Transaction(LOCTEXT("SocketManager_DuplicateSocket", "Duplicate Socket"));
+
+	UStaticMesh* CurrentStaticMesh = StaticMeshEditorPinned->GetStaticMesh();
+
+	CurrentStaticMesh->PreEditChange(NULL);
+
+	TArray<UStaticMeshSocket*> NewSockets;
+	NewSockets.Reserve(SelectedSockets.Num());
+	for (UStaticMeshSocket* SelectedSocket : SelectedSockets)
+	{
+		UStaticMeshSocket* NewSocket = DuplicateObject(SelectedSocket, CurrentStaticMesh);
+		NewSockets.Add(NewSocket);
+
+		// Create a unique name for this socket
+		NewSocket->SocketName = MakeUniqueObjectName(CurrentStaticMesh, UStaticMeshSocket::StaticClass(), NewSocket->SocketName);
+
+		// Add the new socket to the static mesh
+		CurrentStaticMesh->AddSocket(NewSocket);
+	}
+
+	CurrentStaticMesh->PostEditChange();
+	CurrentStaticMesh->MarkPackageDirty();
+
+	RefreshSocketList();
+
+	// Select the duplicated sockets
+	SetSelectedSockets(NewSockets);
+}
 
 void SSocketManager::UpdateStaticMesh()
 {
 	RefreshSocketList();
 }
 
+void SSocketManager::SetSelectedSockets(TArray<UStaticMeshSocket*>& InSelectedSockets)
+{
+	if (InSelectedSockets.Num())
+	{
+		TArray<TSharedPtr<SocketListItem>> SocketItems;
+
+		for (const UStaticMeshSocket* Socket : InSelectedSockets)
+		{
+			for (int32 i = 0; i < SocketListView->GetItems().Num(); ++i)
+			{
+				if (SocketListView->GetItems()[i]->Socket == Socket)
+				{
+					SocketItems.Add(SocketListView->GetItems()[i]);
+				}
+			}
+		}
+
+		SocketListView->ClearSelection();
+		SocketListView->SetItemSelection(SocketItems, true);
+
+		SocketListView->RequestListRefresh();
+		SocketSelectionChanged(InSelectedSockets);
+	}
+	else
+	{
+		SocketListView->ClearSelection();
+
+		SocketListView->RequestListRefresh();
+
+		SocketSelectionChanged(NULL);
+	}
+}
+
 void SSocketManager::RequestRenameSelectedSocket()
 {
 	if(SocketListView->GetSelectedItems().Num() == 1)
@@ -586,6 +708,31 @@ void SSocketManager::DeleteSelectedSocket()
 	}
 }
 
+void SSocketManager::DeleteSelectedSockets()
+{
+	if(SocketListView->GetSelectedItems().Num())
+	{
+		const FScopedTransaction Transaction(LOCTEXT("DeleteSockets", "Delete Sockets"));
+
+		TSharedPtr<IStaticMeshEditor> StaticMeshEditorPinned = StaticMeshEditorPtr.Pin();
+		if (StaticMeshEditorPinned.IsValid())
+		{
+			UStaticMesh* CurrentStaticMesh = StaticMeshEditorPinned->GetStaticMesh();
+			CurrentStaticMesh->PreEditChange(NULL);
+			for (int32 i = 0; i < SocketListView->GetSelectedItems().Num(); ++i)
+			{
+				UStaticMeshSocket* SelectedSocket = SocketListView->GetSelectedItems()[i]->Socket;
+				SelectedSocket->OnPropertyChanged().RemoveAll(this);
+				CurrentStaticMesh->Sockets.Remove(SelectedSocket);
+			}
+
+			CurrentStaticMesh->PostEditChange();
+
+			RefreshSocketList();
+		}
+	}
+}
+
 void SSocketManager::RefreshSocketList()
 {
 	// The static mesh might not be the same one we built the SocketListView with
@@ -664,20 +811,42 @@ bool SSocketManager::CheckForDuplicateSocket(const FString& InSocketName)
 
 void SSocketManager::SocketSelectionChanged(UStaticMeshSocket* InSocket)
 {
-	TArray<UObject*> SelectedObject;
+	TArray<UStaticMeshSocket*> SelectedSockets = GetSelectedSockets();
+	
+	TArray<UObject*> SelectedObjects;
+	SelectedObjects.Reserve(SelectedSockets.Num());
+	
+	for (UStaticMeshSocket* Socket : SelectedSockets)
+	{
+		SelectedObjects.Add(Socket);
+	}
+
+	SocketDetailsView->SetObjects(SelectedObjects);
+
+	// Notify listeners
+	OnSocketSelectionChanged.ExecuteIfBound();
+}
+
+void SSocketManager::SocketSelectionChanged(TArray<UStaticMeshSocket*> InSockets)
+{
+	TArray<UObject*> SelectedObjects;
+	SelectedObjects.Reserve(InSockets.Num());
 
-	if(InSocket)
+	for (UStaticMeshSocket* Socket : InSockets)
 	{
-		SelectedObject.Add(InSocket);
+		if (Socket)
+		{
+			SelectedObjects.Add(Socket);
+		}
 	}
 
-	SocketDetailsView->SetObjects(SelectedObject);
+	SocketDetailsView->SetObjects(SelectedObjects);
 
 	// Notify listeners
 	OnSocketSelectionChanged.ExecuteIfBound();
 }
 
-void SSocketManager::SocketSelectionChanged_Execute( TSharedPtr<SocketListItem> InItem, ESelectInfo::Type /*SelectInfo*/ )
+void SSocketManager::SocketSelectionChanged_Execute(TSharedPtr<SocketListItem> InItem, ESelectInfo::Type /*SelectInfo*/)
 {
 	if(InItem.IsValid())
 	{
@@ -795,11 +964,13 @@ void SSocketManager::OnSocketPropertyChanged( const UStaticMeshSocket* Socket, c
 
 	if ( ChangedPropertyName == RelativeRotationName )
 	{
-		const UStaticMeshSocket* SelectedSocket = GetSelectedSocket();
-
-		if( Socket == SelectedSocket )
+		TArray<UStaticMeshSocket*> SelectedSockets = GetSelectedSockets();
+		for (const UStaticMeshSocket* SelectedSocket : SelectedSockets)
 		{
-			WorldSpaceRotation.Set( Socket->RelativeRotation.Pitch, Socket->RelativeRotation.Yaw, Socket->RelativeRotation.Roll );
+			if (Socket == SelectedSocket)
+			{
+				WorldSpaceRotation.Set(Socket->RelativeRotation.Pitch, Socket->RelativeRotation.Yaw, Socket->RelativeRotation.Roll);
+			}
 		}
 	}
 
diff --git a/Engine/Source/Editor/UnrealEd/Private/SSocketManager.h b/Engine/Source/Editor/UnrealEd/Private/SSocketManager.h
--- a/Engine/Source/Editor/UnrealEd/Private/SSocketManager.h
+++ b/Engine/Source/Editor/UnrealEd/Private/SSocketManager.h
@@ -37,13 +37,21 @@ public:
 
 	// ISocketManager interface
 	virtual UStaticMeshSocket* GetSelectedSocket() const override;
+	virtual TArray<UStaticMeshSocket*> GetSelectedSockets() const override;
+	virtual bool HasSelectedSockets() const override;
 	virtual void SetSelectedSocket(UStaticMeshSocket* InSelectedSocket) override;
+	virtual void AddSelectedSocket(UStaticMeshSocket* InSelectedSocket) override;
+	virtual void RemoveSelectedSocket(const UStaticMeshSocket* InSelectedSocket) override;
 	virtual void DeleteSelectedSocket() override;
+	virtual void DeleteSelectedSockets() override;
 	virtual void DuplicateSelectedSocket() override;
+	virtual void DuplicateSelectedSockets() override;
 	virtual void RequestRenameSelectedSocket() override;
 	virtual void UpdateStaticMesh() override;
 	// End of ISocketManager
 
+	virtual void SetSelectedSockets(TArray<UStaticMeshSocket*>& InSelectedSockets);
+
 		/**
 	 *	Checks for a duplicate socket using the name for comparison.
 	 *
@@ -72,6 +80,7 @@ private:
 	 *	@param InSocket				The newly selected socket.
 	 */
 	void SocketSelectionChanged(UStaticMeshSocket* InSocket);
+	void SocketSelectionChanged(TArray<UStaticMeshSocket*> InSockets);
 
 	/** Callback for the list view when an item is selected. */
 	void SocketSelectionChanged_Execute( TSharedPtr<SocketListItem> InItem, ESelectInfo::Type SelectInfo );
diff --git a/Engine/Source/Editor/UnrealEd/Public/ISocketManager.h b/Engine/Source/Editor/UnrealEd/Public/ISocketManager.h
--- a/Engine/Source/Editor/UnrealEd/Public/ISocketManager.h
+++ b/Engine/Source/Editor/UnrealEd/Public/ISocketManager.h
@@ -11,17 +11,38 @@ class UStaticMeshSocket;
 class ISocketManager : public SCompoundWidget
 {
 public:
-	/** Retrieves the selected socket. */
+	/** Retrieves the first selected socket. */
+	UE_DEPRECATED(5.6, "GetSelectedSocket is deprecated, use GetSelectedSockets instead.")
 	virtual UStaticMeshSocket* GetSelectedSocket() const = 0;
+	
+	/** Retrieves the selected socket. */
+	virtual TArray<UStaticMeshSocket*> GetSelectedSockets() const = 0;
+
+	/** Returns true if at least one socket is selected. */
+	virtual bool HasSelectedSockets() const = 0;
 
 	/** Sets the selected socket. */
 	virtual void SetSelectedSocket(UStaticMeshSocket* InSelectedSocket) = 0;
 
-	/** Deletes the selected socket. */
+	/** Adds a socket to the already selected ones. */
+	virtual void AddSelectedSocket(UStaticMeshSocket* InSelectedSocket) = 0;
+
+	/** Removes a given socket from the current selection. */
+	virtual void RemoveSelectedSocket(const UStaticMeshSocket* InSelectedSocket) = 0;
+
+	/** Deletes the first selected socket. */
+	UE_DEPRECATED(5.6, "DeleteSelectedSocket is deprecated, use DeleteSelectedSockets instead.")
 	virtual void DeleteSelectedSocket() = 0;
 
-	/** Duplicate the selected socket */
+	/** Deletes the selected sockets. */
+	virtual void DeleteSelectedSockets() = 0;
+
+	/** Duplicate the first selected socket. */
+	UE_DEPRECATED(5.6, "DuplicateSelectedSocket is deprecated, use DuplicateSelectedSockets instead.")
 	virtual void DuplicateSelectedSocket() = 0;
+	
+	/** Duplicate the selected sockets. */
+	virtual void DuplicateSelectedSockets() = 0;
 
 	/** Request a rename on the selected socket */
 	virtual void RequestRenameSelectedSocket() = 0;
