[StaticMesh] Add Nanite import/re-import API to the static mesh sub system. This will allow users to script nanite import/re-import.

#rnx

diff --git a/Engine/Source/Editor/StaticMeshEditor/Private/StaticMeshEditorSubsystem.cpp b/Engine/Source/Editor/StaticMeshEditor/Private/StaticMeshEditorSubsystem.cpp
--- a/Engine/Source/Editor/StaticMeshEditor/Private/StaticMeshEditorSubsystem.cpp
+++ b/Engine/Source/Editor/StaticMeshEditor/Private/StaticMeshEditorSubsystem.cpp
@@ -1084,6 +1084,136 @@ bool UStaticMeshEditorSubsystem::SetLodScreenSizes(UStaticMesh* StaticMesh, cons
 	return true;
 }
 
+bool UStaticMeshEditorSubsystem::ImportNaniteHiResMesh(UStaticMesh* StaticMesh, const FString& SourceFilename, bool bShowDialogWhenFileMissing)
+{
+	if (!StaticMesh)
+	{
+		UE_LOG(LogStaticMeshEditorSubsystem, Error, TEXT("ImportNaniteHiResMesh: The StaticMesh is null."));
+		return false;
+	}
+	FStaticMeshSourceModel& HiResModel = StaticMesh->GetHiResSourceModel();
+	FString ResolveFilename = SourceFilename;
+	if (ResolveFilename.IsEmpty())
+	{
+		//Grab the existing source
+		ResolveFilename = HiResModel.SourceImportFilename;
+	}
+
+	bool bSourceFileExists = FPaths::FileExists(ResolveFilename);
+	if (!bSourceFileExists)
+	{
+		ResolveFilename = ResolveFilename.IsEmpty() ?
+			ResolveFilename :
+			UAssetImportData::ResolveImportFilename(ResolveFilename, nullptr);
+	}
+
+	bSourceFileExists = FPaths::FileExists(ResolveFilename);
+
+	// Close the mesh editor to prevent crashing. Reopen it after the mesh has been built.
+	UAssetEditorSubsystem* AssetEditorSubsystem = GEditor->GetEditorSubsystem<UAssetEditorSubsystem>();
+	bool bStaticMeshIsEdited = false;
+	if (AssetEditorSubsystem->FindEditorForAsset(StaticMesh, false))
+	{
+		AssetEditorSubsystem->CloseAllEditorsForAsset(StaticMesh);
+		bStaticMeshIsEdited = true;
+	}
+
+	const bool OriginalEnableNaniteValue = StaticMesh->NaniteSettings.bEnabled;
+	StaticMesh->NaniteSettings.bEnabled = true;
+	bool bResult = true;
+	constexpr bool bAsyncFalse = false;
+	if(bSourceFileExists)
+	{
+		
+		if (!FbxMeshUtils::ImportStaticMeshHiResSourceModel(StaticMesh, ResolveFilename, bAsyncFalse))
+		{
+			UE_LOG(LogStaticMeshEditorSubsystem, Error, TEXT("ImportNaniteHiResMesh: Cannot import hi res mesh."));
+			bResult = false;
+		}
+	}
+	else if (bShowDialogWhenFileMissing)
+	{
+		if (!FbxMeshUtils::ImportStaticMeshHiResSourceModelDialog(StaticMesh))
+		{
+			UE_LOG(LogStaticMeshEditorSubsystem, Error, TEXT("ImportNaniteHiResMesh: Cannot import hi res mesh with dialog."));
+			bResult = false;
+		}
+	}
+	else
+	{
+		UE_LOG(LogStaticMeshEditorSubsystem, Error, TEXT("ImportNaniteHiResMesh: Cannot import hi res mesh, the provided file is invalid."));
+		bResult = false;
+	}
+	
+	if (!bResult)
+	{
+		StaticMesh->NaniteSettings.bEnabled = OriginalEnableNaniteValue;
+	}
+
+	if (bStaticMeshIsEdited)
+	{
+		AssetEditorSubsystem->OpenEditorForAsset(StaticMesh);
+	}
+
+	return bResult;
+}
+
+FString UStaticMeshEditorSubsystem::GetNaniteSourceFilename(const UStaticMesh* StaticMesh)
+{
+	FString ResolveFilename;
+	if (!StaticMesh)
+	{
+		UE_LOG(LogStaticMeshEditorSubsystem, Error, TEXT("GetNaniteSourceFilename: The StaticMesh is null."));
+		return ResolveFilename;
+	}
+	const FStaticMeshSourceModel& HiResModel = StaticMesh->GetHiResSourceModel();
+	ResolveFilename = HiResModel.SourceImportFilename;
+	bool bSourceFileExists = FPaths::FileExists(ResolveFilename);
+	if (!bSourceFileExists)
+	{
+		ResolveFilename = ResolveFilename.IsEmpty() ?
+			ResolveFilename :
+			UAssetImportData::ResolveImportFilename(ResolveFilename, nullptr);
+	}
+	return ResolveFilename;
+}
+
+bool UStaticMeshEditorSubsystem::UpdateNaniteSourceFilename(UStaticMesh* StaticMesh, const FString& NewSourceFilename)
+{
+	bool bResult = false;
+	FString ResolveFilename = NewSourceFilename;
+	if (!StaticMesh)
+	{
+		UE_LOG(LogStaticMeshEditorSubsystem, Error, TEXT("UpdateNaniteSourceFilename: The StaticMesh is null."));
+		return bResult;
+	}
+	if (!FPaths::FileExists(ResolveFilename))
+	{
+		ResolveFilename = ResolveFilename.IsEmpty() ?
+			ResolveFilename :
+			UAssetImportData::ResolveImportFilename(ResolveFilename, nullptr);
+	}
+	if (!StaticMesh->GetHiResSourceModel().SourceImportFilename.Equals(ResolveFilename))
+	{
+		if (ResolveFilename.IsEmpty())
+		{
+			FbxMeshUtils::RemoveStaticMeshHiRes(StaticMesh);
+			bResult = true;
+		}
+		else if(FPaths::FileExists(ResolveFilename))
+		{
+			//We have a valid path, update it and re-import the nanite mesh
+			StaticMesh->GetHiResSourceModel().SourceImportFilename = ResolveFilename;
+			ImportNaniteHiResMesh(StaticMesh, ResolveFilename, false);
+			bResult = true;
+		}
+		else
+		{
+			UE_LOG(LogStaticMeshEditorSubsystem, Error, TEXT("UpdateNaniteSourceFilename: The new source filename is invalid (it can be empty or must point to an existing filename)."));
+		}
+	}
+	return bResult;
+}
 
 FMeshNaniteSettings UStaticMeshEditorSubsystem::GetNaniteSettings(const UStaticMesh* StaticMesh)
 {
diff --git a/Engine/Source/Editor/StaticMeshEditor/Public/StaticMeshEditorSubsystem.h b/Engine/Source/Editor/StaticMeshEditor/Public/StaticMeshEditorSubsystem.h
--- a/Engine/Source/Editor/StaticMeshEditor/Public/StaticMeshEditorSubsystem.h
+++ b/Engine/Source/Editor/StaticMeshEditor/Public/StaticMeshEditorSubsystem.h
@@ -177,6 +177,37 @@ public:
 	bool SetLodScreenSizes(UStaticMesh* StaticMesh, const TArray<float>& ScreenSizes);
 
 public:
+
+	/**
+	 * Import or re-import the nanite hi-res mesh into the specified base mesh.
+	 *
+	 * @param StaticMesh:	The static mesh we import or re-import the nanite hi-res mesh.
+	 * @param SourceFilename:	The source filename containing the hi-res mesh data. If we are re-importing an existing nanite hi-res mesh,
+	 *							it can be empty, in this case it will use the last import file. Otherwise it must be an existing file.
+	 * @param bShowDialogWhenFileMissing: If true a file picker will show in case the provided filename and the existing source filename are empty or do not exist.
+	 *
+	 * @return true if the nanite hi-res mesh import/re-import is successful, false otherwise.
+	 */
+	UFUNCTION(BlueprintCallable, Category = "Editor Scripting | StaticMesh")
+	bool ImportNaniteHiResMesh(UStaticMesh* StaticMesh, const FString& SourceFilename, bool bShowDialogWhenFileMissing);
+
+	/**
+	* Get the Nanite hi-res source filename
+	* @param StaticMesh        Mesh to access
+	* @return the source filename use by the static mesh nanite hi-res.
+	*/
+	UFUNCTION(BlueprintCallable, Category = "Editor Scripting | StaticMesh")
+	FString GetNaniteSourceFilename(const UStaticMesh* StaticMesh);
+
+	/**
+	* Update the Nanite hi-res source filename.
+	* @param StaticMesh			Mesh to update the nanite source filename
+	* @param NewSourceFilename	The source filename we want to set. If empty, imported nanite mesh will be remove. If not empty, the file must exist otherwise nothing will be done.
+	* @return true if the source filename was updated (cleared or replace by an existing filename), false otherwise.
+	*/
+	UFUNCTION(BlueprintCallable, Category = "Editor Scripting | StaticMesh")
+	bool UpdateNaniteSourceFilename(UStaticMesh* StaticMesh, const FString& NewSourceFilename);
+
 	/**
 	* Get the Nanite Settings for the mesh
 	* @param StaticMesh        Mesh to access
@@ -193,6 +224,8 @@ public:
 	*/
 	UFUNCTION(BlueprintCallable, Category = "Static Mesh Utilities")
 	void SetNaniteSettings(UStaticMesh* StaticMesh, FMeshNaniteSettings NaniteSettings, bool bApplyChanges=true);
+
+
 	
 public:
 	/**
