[StaticMeshEditorSubsystem] Fixing crash when calling some StaticMeshEditorSubsystem without a valid StaticMesh.

diff --git a/Engine/Source/Editor/StaticMeshEditor/Private/StaticMeshEditorSubsystem.cpp b/Engine/Source/Editor/StaticMeshEditor/Private/StaticMeshEditorSubsystem.cpp
--- a/Engine/Source/Editor/StaticMeshEditor/Private/StaticMeshEditorSubsystem.cpp
+++ b/Engine/Source/Editor/StaticMeshEditor/Private/StaticMeshEditorSubsystem.cpp
@@ -593,12 +593,23 @@ void UStaticMeshEditorSubsystem::SetLodBuildSettings(UStaticMesh* StaticMesh, co
 
 FName UStaticMeshEditorSubsystem::GetLODGroup(const UStaticMesh* StaticMesh)
 {
+	if (StaticMesh == nullptr)
+	{
+		UE_LOG(LogStaticMeshEditorSubsystem, Error, TEXT("StaticMesh GetLODGroup: The StaticMesh is null."));
+		return NAME_None;
+	}
+
 	return StaticMesh->LODGroup;
 }
 
 bool UStaticMeshEditorSubsystem::SetLODGroup(UStaticMesh* StaticMesh, FName LODGroup, bool bRebuildImmediately)
 {
-	
+	if (StaticMesh == nullptr)
+	{
+		UE_LOG(LogStaticMeshEditorSubsystem, Error, TEXT("StaticMesh SetLODGroup: The StaticMesh is null."));
+		return false;
+	}
+
 	TArray<FName> LODGroups;
 	UStaticMesh::GetLODGroups(LODGroups);
 	
