Changed "Show Collision" modes in the static mesh editor to not draw any collision if we will never cook collision data for runtime.

This should reduce confusion about whether collision data is built for a mesh as previously even if it wasn't built it would be shown when requested. Now it will not show if we never intend to cook the collison data and a warning is added to the static mesh editor viewport information to highlight that while a mesh may have collision primitives, its collision will never be cooked.

diff --git a/Engine/Source/Editor/StaticMeshEditor/Private/StaticMeshEditorViewportClient.cpp b/Engine/Source/Editor/StaticMeshEditor/Private/StaticMeshEditorViewportClient.cpp
--- a/Engine/Source/Editor/StaticMeshEditor/Private/StaticMeshEditorViewportClient.cpp
+++ b/Engine/Source/Editor/StaticMeshEditor/Private/StaticMeshEditorViewportClient.cpp
@@ -496,8 +496,11 @@ void FStaticMeshEditorViewportClient::Draw(const FSceneView* View,FPrimitiveDraw
 		return;
 	}
 
+	UBodySetup* MeshBodySetup = StaticMesh->GetBodySetup();
+	const bool bHasCookedPhysicsState = MeshBodySetup && !MeshBodySetup->bNeverNeedsCookedCollisionData;
+
 	// Draw simple shapes if we are showing simple, or showing complex but using simple as complex
-	if (StaticMesh->GetBodySetup() && (bShowSimpleCollision || (bShowComplexCollision && StaticMesh->GetBodySetup()->CollisionTraceFlag == ECollisionTraceFlag::CTF_UseSimpleAsComplex)))
+	if (StaticMesh->GetBodySetup() && bHasCookedPhysicsState && (bShowSimpleCollision || (bShowComplexCollision && StaticMesh->GetBodySetup()->CollisionTraceFlag == ECollisionTraceFlag::CTF_UseSimpleAsComplex)))
 	{
 		// Ensure physics mesh is created before we try and draw it
 		StaticMesh->GetBodySetup()->CreatePhysicsMeshes();
@@ -576,7 +579,7 @@ void FStaticMeshEditorViewportClient::Draw(const FSceneView* View,FPrimitiveDraw
 		}
 	}
 
-	if (bShowComplexCollision && StaticMesh->ComplexCollisionMesh && StaticMesh->GetBodySetup()->CollisionTraceFlag != ECollisionTraceFlag::CTF_UseSimpleAsComplex)
+	if (bShowComplexCollision && bHasCookedPhysicsState && StaticMesh->ComplexCollisionMesh && StaticMesh->GetBodySetup()->CollisionTraceFlag != ECollisionTraceFlag::CTF_UseSimpleAsComplex)
 	{
 		const FColor UnselectedColor(0, 0, 125);
 
@@ -960,7 +963,13 @@ void FStaticMeshEditorViewportClient::DrawCanvas(FViewport& InViewport, FSceneVi
 	// Show the number of collision primitives
 	if (StaticMesh->GetBodySetup())
 	{
-		TextItems.Emplace(FText::Format(NSLOCTEXT("UnrealEd", "NumPrimitives_F", "Num Collision Primitives:  {0}"), FText::AsNumber(StaticMesh->GetBodySetup()->AggGeom.GetElementCount())));
+		UBodySetup* MeshBodySetup = StaticMesh->GetBodySetup();
+		TextItems.Emplace(FText::Format(NSLOCTEXT("UnrealEd", "NumPrimitives_F", "Num Collision Primitives:  {0}"), FText::AsNumber(MeshBodySetup->AggGeom.GetElementCount())));
+
+		if(MeshBodySetup->bNeverNeedsCookedCollisionData)
+		{
+			TextItems.Emplace(NSLOCTEXT("UnrealEd", "NumPrimitivesCookWarning", "\tbNeverNeedsCookedCollisionData Enabled, no collision data will be used."), true);
+		}
 	}
 
 	// Estimated compressed size
diff --git a/Engine/Source/Runtime/Engine/Private/StaticMeshSceneProxy.cpp b/Engine/Source/Runtime/Engine/Private/StaticMeshSceneProxy.cpp
--- a/Engine/Source/Runtime/Engine/Private/StaticMeshSceneProxy.cpp
+++ b/Engine/Source/Runtime/Engine/Private/StaticMeshSceneProxy.cpp
@@ -2750,10 +2750,6 @@ void FStaticMeshSceneProxyDesc::InitializeFromStaticMeshComponent(const UStaticM
 	WorldPositionOffsetDisableDistance = InComponent->WorldPositionOffsetDisableDistance ;
 	NanitePixelProgrammableDistance = InComponent->NanitePixelProgrammableDistance;
 	bReverseCulling = InComponent->bReverseCulling ;
-#if STATICMESH_ENABLE_DEBUG_RENDERING
-	bDrawMeshCollisionIfComplex = InComponent->bDrawMeshCollisionIfComplex ;
-	bDrawMeshCollisionIfSimple = InComponent->bDrawMeshCollisionIfSimple ;
-#endif
 	bEvaluateWorldPositionOffset = InComponent->bEvaluateWorldPositionOffset ;
 	bOverrideMinLOD = InComponent->bOverrideMinLOD ;
 	bCastDistanceFieldIndirectShadow = InComponent->bCastDistanceFieldIndirectShadow ;
@@ -2785,6 +2781,12 @@ void FStaticMeshSceneProxyDesc::InitializeFromStaticMeshComponent(const UStaticM
 	NaniteResources = InComponent->GetNaniteResources();
 	BodySetup = const_cast<UStaticMeshComponent*>(InComponent)->GetBodySetup();
 
+#if STATICMESH_ENABLE_DEBUG_RENDERING
+	const bool bHasCollisionState = BodySetup && !BodySetup->bNeverNeedsCookedCollisionData;
+	bDrawMeshCollisionIfComplex = InComponent->bDrawMeshCollisionIfComplex && bHasCollisionState;
+	bDrawMeshCollisionIfSimple = InComponent->bDrawMeshCollisionIfSimple && bHasCollisionState;
+#endif
+
 	LODData = const_cast<UStaticMeshComponent*>(InComponent)->LODData;
 
 	WireframeColor = InComponent->GetWireframeColor();
