Make sure Material list copy paste paste the overlay material

#rnx

diff --git a/Engine/Source/Editor/Persona/Private/PersonaMeshDetails.cpp b/Engine/Source/Editor/Persona/Private/PersonaMeshDetails.cpp
--- a/Engine/Source/Editor/Persona/Private/PersonaMeshDetails.cpp
+++ b/Engine/Source/Editor/Persona/Private/PersonaMeshDetails.cpp
@@ -2622,6 +2622,7 @@ void FPersonaMeshDetails::OnPasteMaterialList()
 				if (MeshMaterials.IsValidIndex(MaterialIndex))
 				{
 					MeshMaterials[MaterialIndex].MaterialInterface = TempMaterials[MaterialIndex].MaterialInterface;
+					MeshMaterials[MaterialIndex].OverlayMaterialInterface = TempMaterials[MaterialIndex].OverlayMaterialInterface;
 				}
 			}
 			
@@ -2696,6 +2697,7 @@ void FPersonaMeshDetails::OnPasteMaterialItem(int32 CurrentSlot)
 				FSkeletalMaterial TmpSkeletalMaterial;
 				FJsonObjectConverter::JsonObjectToUStruct(RootJsonObject.ToSharedRef(), FSkeletalMaterial::StaticStruct(), &TmpSkeletalMaterial, 0, 0);
 				Mesh->GetMaterials()[CurrentSlot].MaterialInterface = TmpSkeletalMaterial.MaterialInterface;
+				Mesh->GetMaterials()[CurrentSlot].OverlayMaterialInterface = TmpSkeletalMaterial.OverlayMaterialInterface;
 			}
 
 			Mesh->PostEditChange();
diff --git a/Engine/Source/Editor/StaticMeshEditor/Private/StaticMeshEditorTools.cpp b/Engine/Source/Editor/StaticMeshEditor/Private/StaticMeshEditorTools.cpp
--- a/Engine/Source/Editor/StaticMeshEditor/Private/StaticMeshEditorTools.cpp
+++ b/Engine/Source/Editor/StaticMeshEditor/Private/StaticMeshEditorTools.cpp
@@ -2851,6 +2851,7 @@ void FMeshMaterialsLayout::OnPasteMaterialList()
 			if (GetStaticMesh().GetStaticMaterials().IsValidIndex(MaterialIndex))
 			{
 				GetStaticMesh().GetStaticMaterials()[MaterialIndex].MaterialInterface = TempMaterials[MaterialIndex].MaterialInterface;
+				GetStaticMesh().GetStaticMaterials()[MaterialIndex].OverlayMaterialInterface = TempMaterials[MaterialIndex].OverlayMaterialInterface;
 			}
 		}
 
@@ -2911,6 +2912,7 @@ void FMeshMaterialsLayout::OnPasteMaterialItem(int32 CurrentSlot)
 			FStaticMaterial TmpStaticMaterial;
 			FJsonObjectConverter::JsonObjectToUStruct(RootJsonObject.ToSharedRef(), FStaticMaterial::StaticStruct(), &TmpStaticMaterial, 0, 0);
 			GetStaticMesh().GetStaticMaterials()[CurrentSlot].MaterialInterface = TmpStaticMaterial.MaterialInterface;
+			GetStaticMesh().GetStaticMaterials()[CurrentSlot].OverlayMaterialInterface = TmpStaticMaterial.OverlayMaterialInterface;
 		}
 
 		CallPostEditChange(Property);
