Updated headers using UnrealCodeFixup to make sure dllstorage is on methods/staticvars instead of types. Part 5/n

diff --git a/Engine/Source/Editor/SequencerCore/Public/MVVM/Views/STrackLane.h b/Engine/Source/Editor/SequencerCore/Public/MVVM/Views/STrackLane.h
--- a/Engine/Source/Editor/SequencerCore/Public/MVVM/Views/STrackLane.h
+++ b/Engine/Source/Editor/SequencerCore/Public/MVVM/Views/STrackLane.h
@@ -12,6 +12,8 @@
 #include "MVVM/Extensions/ITrackLaneExtension.h"
 #include "MVVM/Extensions/IOutlinerExtension.h"
 
+#define UE_API SEQUENCERCORE_API
+
 class FPaintArgs;
 class FSlateWindowElementList;
 
@@ -26,47 +28,47 @@ class STrackAreaView;
 /**
  * A wrapper widget responsible for positioning a track lane within the section area
  */
-class SEQUENCERCORE_API STrackLane
+class STrackLane
 	: public SPanel
 {
 public:
 	SLATE_BEGIN_ARGS(STrackLane){}
 	SLATE_END_ARGS()
 
-	STrackLane();
-	~STrackLane();
+	UE_API STrackLane();
+	UE_API ~STrackLane();
 
 	/** Construct this widget */
-	void Construct(const FArguments& InArgs, TWeakPtr<STrackAreaView> InTrackAreaView, TWeakViewModelPtr<IOutlinerExtension> InWeakOutlinerItem, TSharedPtr<STrackLane> InParentLane, const FTrackAreaParameters& InTrackParams, const TSharedRef<SOutlinerView>& InTreeView);
+	UE_API void Construct(const FArguments& InArgs, TWeakPtr<STrackAreaView> InTrackAreaView, TWeakViewModelPtr<IOutlinerExtension> InWeakOutlinerItem, TSharedPtr<STrackLane> InParentLane, const FTrackAreaParameters& InTrackParams, const TSharedRef<SOutlinerView>& InTreeView);
 
 	/*~ SWidget Interface */
-	int32 OnPaint(const FPaintArgs& Args, const FGeometry& AllottedGeometry, const FSlateRect& MyCullingRect, FSlateWindowElementList& OutDrawElements, int32 LayerId, const FWidgetStyle& InWidgetStyle, bool bParentEnabled) const override;
-	void Tick(const FGeometry& AllottedGeometry, const double InCurrentTime, const float InDeltaTime) override;
-	FReply OnMouseButtonDown(const FGeometry& MyGeometry, const FPointerEvent& MouseEvent) override;
-	FReply OnMouseButtonUp(const FGeometry& MyGeometry, const FPointerEvent& MouseEvent) override;
-	FReply OnMouseMove(const FGeometry& MyGeometry, const FPointerEvent& MouseEvent) override;
-	FCursorReply OnCursorQuery(const FGeometry& MyGeometry, const FPointerEvent& CursorEvent) const override;
+	UE_API int32 OnPaint(const FPaintArgs& Args, const FGeometry& AllottedGeometry, const FSlateRect& MyCullingRect, FSlateWindowElementList& OutDrawElements, int32 LayerId, const FWidgetStyle& InWidgetStyle, bool bParentEnabled) const override;
+	UE_API void Tick(const FGeometry& AllottedGeometry, const double InCurrentTime, const float InDeltaTime) override;
+	UE_API FReply OnMouseButtonDown(const FGeometry& MyGeometry, const FPointerEvent& MouseEvent) override;
+	UE_API FReply OnMouseButtonUp(const FGeometry& MyGeometry, const FPointerEvent& MouseEvent) override;
+	UE_API FReply OnMouseMove(const FGeometry& MyGeometry, const FPointerEvent& MouseEvent) override;
+	UE_API FCursorReply OnCursorQuery(const FGeometry& MyGeometry, const FPointerEvent& CursorEvent) const override;
 
 	/*~ SPanel Interface */
-	void OnArrangeChildren( const FGeometry& AllottedGeometry, FArrangedChildren& ArrangedChildren ) const override;
-	FVector2D ComputeDesiredSize(float) const override;
+	UE_API void OnArrangeChildren( const FGeometry& AllottedGeometry, FArrangedChildren& ArrangedChildren ) const override;
+	UE_API FVector2D ComputeDesiredSize(float) const override;
 	FChildren* GetChildren() override { return &Children; }
 
 	/** Gets the outliner view-model that created this track lane */
-	TViewModelPtr<IOutlinerExtension> GetOutlinerItem() const;
+	UE_API TViewModelPtr<IOutlinerExtension> GetOutlinerItem() const;
 
-	TSharedPtr<STrackAreaView> GetTrackAreaView() const;
-	TSharedPtr<STrackLane> GetParentLane() const;
-	TSharedPtr<FTimeToPixel> GetTimeToPixel() const;
+	UE_API TSharedPtr<STrackAreaView> GetTrackAreaView() const;
+	UE_API TSharedPtr<STrackLane> GetParentLane() const;
+	UE_API TSharedPtr<FTimeToPixel> GetTimeToPixel() const;
 
-	TArrayView<const TWeakPtr<STrackLane>> GetChildLanes() const;
+	UE_API TArrayView<const TWeakPtr<STrackLane>> GetChildLanes() const;
 
 	const TMap<FWeakViewModelPtr, TSharedPtr<ITrackLaneWidget>>& GetAllWidgets() const { return WidgetsByModel; }
 
 	/** Gets the vertical offset of this track lane */
-	float GetVerticalPosition() const;
+	UE_API float GetVerticalPosition() const;
 	/** Sets the vertical offset of this track lane */
-	void SetVerticalPosition(float InPosition);
+	UE_API void SetVerticalPosition(float InPosition);
 
 	/**
 	 * Sets the vertical offset of this track lane and all parent track lanes given the vertical offset of
@@ -74,19 +76,19 @@ public:
 	 * item that created this track lane)
 	 * Note that the given item *could* actually be the one that created this track lane.
 	 */
-	void PositionParentTrackLanes(TViewModelPtr<IOutlinerExtension> InItem, float InItemTop);
+	UE_API void PositionParentTrackLanes(TViewModelPtr<IOutlinerExtension> InItem, float InItemTop);
 
 	/** Returns whether the outliner item that created this track lane is pinned or not */
-	bool IsPinned() const;
+	UE_API bool IsPinned() const;
 
-	TSharedPtr<ITrackLaneWidget> FindWidgetForModel(const FWeakViewModelPtr& InModel) const;
+	UE_API TSharedPtr<ITrackLaneWidget> FindWidgetForModel(const FWeakViewModelPtr& InModel) const;
 
 protected:
 
-	int32 PaintLaneForeground(const FGeometry& AllottedGeometry, const FSlateRect& MyCullingRect, FSlateWindowElementList& OutDrawElements, int32 LayerId, const FWidgetStyle& InWidgetStyle) const;
+	UE_API int32 PaintLaneForeground(const FGeometry& AllottedGeometry, const FSlateRect& MyCullingRect, FSlateWindowElementList& OutDrawElements, int32 LayerId, const FWidgetStyle& InWidgetStyle) const;
 
-	void OnHierarchyUpdated();
-	void RecreateWidgets();
+	UE_API void OnHierarchyUpdated();
+	UE_API void RecreateWidgets();
 
 private:
 	TWeakPtr<STrackAreaView> WeakTrackAreaView;
@@ -139,3 +141,4 @@ private:
 } // namespace Sequencer
 } // namespace UE
 
+#undef UE_API
diff --git a/Engine/Source/Editor/SequencerWidgets/Public/ITimeSlider.h b/Engine/Source/Editor/SequencerWidgets/Public/ITimeSlider.h
--- a/Engine/Source/Editor/SequencerWidgets/Public/ITimeSlider.h
+++ b/Engine/Source/Editor/SequencerWidgets/Public/ITimeSlider.h
@@ -6,21 +6,6 @@
 #include "ViewRangeInterpolation.h"
 #include "Widgets/SCompoundWidget.h"
 
-#if UE_ENABLE_INCLUDE_ORDER_DEPRECATED_IN_5_3
-#include "CoreMinimal.h"
-#include "Misc/Attribute.h"
-#include "Input/CursorReply.h"
-#include "Widgets/SWidget.h"
-#include "Widgets/DeclarativeSyntaxSupport.h"
-#include "Widgets/SCompoundWidget.h"
-#include "Widgets/Input/NumericTypeInterface.h"
-#include "ISequencerInputHandler.h"
-#include "IMovieScenePlayer.h"
-#include "MovieScene.h"
-
-#include "AnimatedRange.h"
-#include "TimeSliderArgs.h"
-#endif
 
 class FSlateWindowElementList;
 struct FSlateBrush;
diff --git a/Engine/Source/Editor/SequencerWidgets/Public/WidgetFocusUtils.h b/Engine/Source/Editor/SequencerWidgets/Public/WidgetFocusUtils.h
--- a/Engine/Source/Editor/SequencerWidgets/Public/WidgetFocusUtils.h
+++ b/Engine/Source/Editor/SequencerWidgets/Public/WidgetFocusUtils.h
@@ -6,6 +6,8 @@
 #include "Templates/Function.h"
 #include "Delegates/IDelegateInstance.h"
 
+#define UE_API SEQUENCERWIDGETS_API
+
 class SWidget;
 struct FKeyEvent;
 struct FPointerEvent;
@@ -17,28 +19,28 @@ struct FPointerEvent;
  * If the mouse leave event is called without any key down event having been called, the function is reset and the focus is not modified at all.
  */
 
-class SEQUENCERWIDGETS_API FPendingWidgetFocus : public TSharedFromThis<FPendingWidgetFocus>
+class FPendingWidgetFocus : public TSharedFromThis<FPendingWidgetFocus>
 {
 public:
 
 	FPendingWidgetFocus() = default;
-	FPendingWidgetFocus(const TArray<FName>& InTypesKeepingFocus);
+	UE_API FPendingWidgetFocus(const TArray<FName>& InTypesKeepingFocus);
 
-	static FPendingWidgetFocus MakeNoTextEdit();
+	static UE_API FPendingWidgetFocus MakeNoTextEdit();
 	
-	~FPendingWidgetFocus();
+	UE_API ~FPendingWidgetFocus();
 
-	void Enable(const bool InEnabled);
-	bool IsEnabled() const;
+	UE_API void Enable(const bool InEnabled);
+	UE_API bool IsEnabled() const;
 
-	void SetPendingFocusIfNeeded(const TWeakPtr<SWidget>& InWidget);
-	void ResetPendingFocus();
+	UE_API void SetPendingFocusIfNeeded(const TWeakPtr<SWidget>& InWidget);
+	UE_API void ResetPendingFocus();
 
 private:
 	
-	void OnPreInputKeyDown(const FKeyEvent&);
-	void OnPreInputButtonDown(const FPointerEvent&);
-	bool CanFocusBeStolen() const;
+	UE_API void OnPreInputKeyDown(const FKeyEvent&);
+	UE_API void OnPreInputButtonDown(const FPointerEvent&);
+	UE_API bool CanFocusBeStolen() const;
 
 	TFunction<void()> PendingFocusFunction;
 	FDelegateHandle PreInputKeyDownHandle;
@@ -46,3 +48,5 @@ private:
 
 	TArray<FName> KeepingFocus;
 };
+
+#undef UE_API
diff --git a/Engine/Source/Editor/SerializedRecorderInterface/Public/ISerializedRecorder.h b/Engine/Source/Editor/SerializedRecorderInterface/Public/ISerializedRecorder.h
--- a/Engine/Source/Editor/SerializedRecorderInterface/Public/ISerializedRecorder.h
+++ b/Engine/Source/Editor/SerializedRecorderInterface/Public/ISerializedRecorder.h
@@ -4,16 +4,19 @@
 #include "CoreMinimal.h"
 #include "Features/IModularFeature.h"
 
+#define UE_API SERIALIZEDRECORDERINTERFACE_API
+
 
 class UMovieSceneSequence;
 class UWorld;
 
-class SERIALIZEDRECORDERINTERFACE_API  ISerializedRecorder : public IModularFeature
+class  ISerializedRecorder : public IModularFeature
 {
 public:
 	virtual ~ISerializedRecorder() {}
 
-	static FName ModularFeatureName;
+	static UE_API FName ModularFeatureName;
 	virtual bool LoadRecordedSequencerFile(UMovieSceneSequence* InMovieSceneSequence, UWorld *PlaybackContext, const FString& InFileName, TFunction<void()> InCompletionCallback) = 0;
 };
 
+#undef UE_API
diff --git a/Engine/Source/Editor/SkeletalMeshEditor/Public/SkeletalMeshEditorContextMenuContext.h b/Engine/Source/Editor/SkeletalMeshEditor/Public/SkeletalMeshEditorContextMenuContext.h
--- a/Engine/Source/Editor/SkeletalMeshEditor/Public/SkeletalMeshEditorContextMenuContext.h
+++ b/Engine/Source/Editor/SkeletalMeshEditor/Public/SkeletalMeshEditorContextMenuContext.h
@@ -6,8 +6,8 @@
 #include "ToolMenuContext.h"
 #include "SkeletalMeshEditorContextMenuContext.generated.h"
 
-UCLASS()
-class SKELETALMESHEDITOR_API USkeletalMeshEditorContextMenuContext : public UObject
+UCLASS(MinimalAPI)
+class USkeletalMeshEditorContextMenuContext : public UObject
 {
 	GENERATED_BODY()
 
@@ -17,4 +17,4 @@ public:
 
 	// The section index that was clicked
 	int32 SectionIndex = INDEX_NONE;
-};
\ No newline at end of file
+};
diff --git a/Engine/Source/Editor/SkeletonEditor/Public/BoneProxy.h b/Engine/Source/Editor/SkeletonEditor/Public/BoneProxy.h
--- a/Engine/Source/Editor/SkeletonEditor/Public/BoneProxy.h
+++ b/Engine/Source/Editor/SkeletonEditor/Public/BoneProxy.h
@@ -10,6 +10,8 @@
 #include "IPersonaPreviewScene.h"
 #include "BoneProxy.generated.h"
 
+#define UE_API SKELETONEDITOR_API
+
 class UDebugSkelMeshComponent;
 enum class ETransformType : uint8;
 namespace ESlateTransformComponent { enum Type : int; }
@@ -18,13 +20,13 @@ namespace ESlateTransformSubComponent { enum Type : int; }
 namespace ETextCommit { enum Type : int; }
 
 /** Proxy object used to display/edit bone transforms */
-UCLASS()
-class SKELETONEDITOR_API UBoneProxy : public UObject, public FTickableEditorObject
+UCLASS(MinimalAPI)
+class UBoneProxy : public UObject, public FTickableEditorObject
 {
 	GENERATED_BODY()
 
 public:
-	UBoneProxy();
+	UE_API UBoneProxy();
 
 	enum ETransformType
 	{
@@ -41,17 +43,17 @@ public:
 
 	
 	/** UObject interface */
-	virtual void PreEditChange(FEditPropertyChain& PropertyAboutToChange) override;
-	virtual void PostEditChangeProperty(FPropertyChangedEvent& PropertyChangedEvent) override;
+	UE_API virtual void PreEditChange(FEditPropertyChain& PropertyAboutToChange) override;
+	UE_API virtual void PostEditChangeProperty(FPropertyChangedEvent& PropertyChangedEvent) override;
 
 	// overload based on a property name
-	void OnPreEditChange(FName PropertyName);
-	void OnPostEditChangeProperty(FName PropertyName);
+	UE_API void OnPreEditChange(FName PropertyName);
+	UE_API void OnPostEditChangeProperty(FName PropertyName);
 	
 	/** FTickableEditorObject interface */
-	virtual void Tick(float DeltaTime) override;
-	virtual bool IsTickable() const override;
-	virtual TStatId GetStatId() const override;
+	UE_API virtual void Tick(float DeltaTime) override;
+	UE_API virtual bool IsTickable() const override;
+	UE_API virtual TStatId GetStatId() const override;
 
 	/** The name of the bone we have selected */
 	UPROPERTY(VisibleAnywhere, Category = "Bone")
@@ -123,14 +125,14 @@ public:
 	bool bIsTransformEditable;
 
 	/** Method to react to retrieval of numeric values for the widget */
-	TOptional<FVector::FReal> GetNumericValue(
+	UE_API TOptional<FVector::FReal> GetNumericValue(
 		ESlateTransformComponent::Type Component,
 		ESlateRotationRepresentation::Type Representation,
 		ESlateTransformSubComponent::Type SubComponent,
 		ETransformType TransformType) const;
 
 	/** Static method to retrieve a value for a list of proxies */
-	static TOptional<FVector::FReal> GetMultiNumericValue(
+	static UE_API TOptional<FVector::FReal> GetMultiNumericValue(
 		ESlateTransformComponent::Type Component,
 		ESlateRotationRepresentation::Type Representation,
 		ESlateTransformSubComponent::Type SubComponent,
@@ -138,14 +140,14 @@ public:
 		TArrayView<UBoneProxy*> BoneProxies);
 
 	/** Method to react to changes of numeric values in the widget */
-	void OnNumericValueChanged(
+	UE_API void OnNumericValueChanged(
 		ESlateTransformComponent::Type Component,
 		ESlateRotationRepresentation::Type Representation,
 		ESlateTransformSubComponent::Type SubComponent,
 		FVector::FReal Value,
 		ETransformType TransformType);
 
-	static void OnSliderMovementStateChanged(
+	static UE_API void OnSliderMovementStateChanged(
 		ESlateTransformComponent::Type Component,
 		ESlateRotationRepresentation::Type Representation,
 		ESlateTransformSubComponent::Type SubComponent,
@@ -154,7 +156,7 @@ public:
 		TArrayView<UBoneProxy*> BoneProxies);
 
 	/** Method to react to changes of numeric values in the widget */
-	static void OnMultiNumericValueChanged(
+	static UE_API void OnMultiNumericValueChanged(
 		ESlateTransformComponent::Type Component,
 		ESlateRotationRepresentation::Type Representation,
 		ESlateTransformSubComponent::Type SubComponent,
@@ -165,12 +167,14 @@ public:
 		TArrayView<UBoneProxy*> BoneProxies);
 
 	// Returns true if a given transform component differs from the reference
-	bool DiffersFromDefault(ESlateTransformComponent::Type Component, ETransformType TransformType) const;
+	UE_API bool DiffersFromDefault(ESlateTransformComponent::Type Component, ETransformType TransformType) const;
 
 	// Resets the bone transform's component to its reference
-	void ResetToDefault(ESlateTransformComponent::Type InComponent, ETransformType TransformType);
+	UE_API void ResetToDefault(ESlateTransformComponent::Type InComponent, ETransformType TransformType);
 
 private:
-	static void BeginSetValueTransaction(ESlateTransformComponent::Type InComponent, TArrayView<UBoneProxy*> BoneProxies);
-	static void EndTransaction();
+	static UE_API void BeginSetValueTransaction(ESlateTransformComponent::Type InComponent, TArrayView<UBoneProxy*> BoneProxies);
+	static UE_API void EndTransaction();
 };
+
+#undef UE_API
diff --git a/Engine/Source/Editor/SkeletonEditor/Public/ISkeletonTreeBuilder.h b/Engine/Source/Editor/SkeletonEditor/Public/ISkeletonTreeBuilder.h
--- a/Engine/Source/Editor/SkeletonEditor/Public/ISkeletonTreeBuilder.h
+++ b/Engine/Source/Editor/SkeletonEditor/Public/ISkeletonTreeBuilder.h
@@ -6,12 +6,14 @@
 #include "Containers/ArrayView.h"
 #include "Misc/TextFilterExpressionEvaluator.h"
 
+#define UE_API SKELETONEDITOR_API
+
 class ISkeletonTreeItem;
 enum class ESkeletonTreeFilterResult : int32;
 class FTextFilterExpressionEvaluator;
 
 /** Output struct for builders to use */
-struct SKELETONEDITOR_API FSkeletonTreeBuilderOutput
+struct FSkeletonTreeBuilderOutput
 {
 	FSkeletonTreeBuilderOutput(TArray<TSharedPtr<class ISkeletonTreeItem>>& InItems, TArray<TSharedPtr<class ISkeletonTreeItem>>& InLinearItems)
 		: Items(InItems)
@@ -25,7 +27,7 @@ struct SKELETONEDITOR_API FSkeletonTreeBuilderOutput
 	 * @param	InParentTypes	The types of items to search. If this is empty all items will be searched.
 	 * @param	bAddToHead		Whether to add the item to the start or end of the parent's children array
 	 */
-	void Add(const TSharedPtr<class ISkeletonTreeItem>& InItem, const FName& InParentName, TArrayView<const FName> InParentTypes, bool bAddToHead = false);
+	UE_API void Add(const TSharedPtr<class ISkeletonTreeItem>& InItem, const FName& InParentName, TArrayView<const FName> InParentTypes, bool bAddToHead = false);
 
 	/** 
 	 * Add an item to the output
@@ -46,7 +48,7 @@ struct SKELETONEDITOR_API FSkeletonTreeBuilderOutput
 	 * @param	InParentType	The type of items to search. If this is NAME_None all items will be searched.
 	 * @param	bAddToHead		Whether to add the item to the start or end of the parent's children array
 	 */
-	void Add(const TSharedPtr<class ISkeletonTreeItem>& InItem, const FName& InParentName, const FName& InParentType, bool bAddToHead = false);
+	UE_API void Add(const TSharedPtr<class ISkeletonTreeItem>& InItem, const FName& InParentName, const FName& InParentType, bool bAddToHead = false);
 
 	/** 
 	 * Find the item with the specified name
@@ -54,7 +56,7 @@ struct SKELETONEDITOR_API FSkeletonTreeBuilderOutput
 	 * @param	InTypes	The types of items to search. If this is empty all items will be searched.
 	 * @return the item found, or an invalid ptr if it was not found.
 	 */
-	TSharedPtr<class ISkeletonTreeItem> Find(const FName& InName, TArrayView<const FName> InTypes) const;
+	UE_API TSharedPtr<class ISkeletonTreeItem> Find(const FName& InName, TArrayView<const FName> InTypes) const;
 
 	/** 
 	 * Find the item with the specified name
@@ -73,7 +75,7 @@ struct SKELETONEDITOR_API FSkeletonTreeBuilderOutput
 	 * @param	InType	The type of items to search. If this is NAME_None all items will be searched.
 	 * @return the item found, or an invalid ptr if it was not found.
 	 */
-	TSharedPtr<class ISkeletonTreeItem> Find(const FName& InName, const FName& InType) const;
+	UE_API TSharedPtr<class ISkeletonTreeItem> Find(const FName& InName, const FName& InType) const;
 
 private:
 	/** The items that are built by this builder */
@@ -132,3 +134,5 @@ public:
 	/** Get whether this builder is showing bones */
 	virtual bool IsShowingAttachedAssets() const = 0;
 };
+
+#undef UE_API
diff --git a/Engine/Source/Editor/SkeletonEditor/Public/SkeletonToolMenuContext.h b/Engine/Source/Editor/SkeletonEditor/Public/SkeletonToolMenuContext.h
--- a/Engine/Source/Editor/SkeletonEditor/Public/SkeletonToolMenuContext.h
+++ b/Engine/Source/Editor/SkeletonEditor/Public/SkeletonToolMenuContext.h
@@ -8,11 +8,11 @@
 
 class ISkeletonEditor;
 
-UCLASS()
-class SKELETONEDITOR_API USkeletonToolMenuContext : public UObject
+UCLASS(MinimalAPI)
+class USkeletonToolMenuContext : public UObject
 {
 	GENERATED_BODY()
 
 public:
 	TWeakPtr<ISkeletonEditor> SkeletonEditor;
-};
\ No newline at end of file
+};
diff --git a/Engine/Source/Editor/SkeletonEditor/Public/SkeletonTreeBuilder.h b/Engine/Source/Editor/SkeletonEditor/Public/SkeletonTreeBuilder.h
--- a/Engine/Source/Editor/SkeletonEditor/Public/SkeletonTreeBuilder.h
+++ b/Engine/Source/Editor/SkeletonEditor/Public/SkeletonTreeBuilder.h
@@ -7,6 +7,8 @@
 #include "IEditableSkeleton.h"
 #include "ISkeletonTreeItem.h"
 
+#define UE_API SKELETONEDITOR_API
+
 class IPersonaPreviewScene;
 class ISkeletonTree;
 class USkeletalMeshSocket;
@@ -56,49 +58,49 @@ enum class ESocketFilter : int32
 	Count
 };
 
-class SKELETONEDITOR_API FSkeletonTreeBuilder : public ISkeletonTreeBuilder
+class FSkeletonTreeBuilder : public ISkeletonTreeBuilder
 {
 public:
-	FSkeletonTreeBuilder(const FSkeletonTreeBuilderArgs& InBuilderArgs);
+	UE_API FSkeletonTreeBuilder(const FSkeletonTreeBuilderArgs& InBuilderArgs);
 
 	/** ISkeletonTreeBuilder interface */
-	virtual void Initialize(const TSharedRef<class ISkeletonTree>& InSkeletonTree, const TSharedPtr<class IPersonaPreviewScene>& InPreviewScene, FOnFilterSkeletonTreeItem InOnFilterSkeletonTreeItem) override;
-	virtual void Build(FSkeletonTreeBuilderOutput& Output) override;
-	virtual void Filter(const FSkeletonTreeFilterArgs& InArgs, const TArray<TSharedPtr<class ISkeletonTreeItem>>& InItems, TArray<TSharedPtr<class ISkeletonTreeItem>>& OutFilteredItems) override;
-	virtual ESkeletonTreeFilterResult FilterItem(const FSkeletonTreeFilterArgs& InArgs, const TSharedPtr<class ISkeletonTreeItem>& InItem) override;
-	virtual bool IsShowingBones() const override;
-	virtual bool IsShowingSockets() const override;
-	virtual bool IsShowingAttachedAssets() const override;
+	UE_API virtual void Initialize(const TSharedRef<class ISkeletonTree>& InSkeletonTree, const TSharedPtr<class IPersonaPreviewScene>& InPreviewScene, FOnFilterSkeletonTreeItem InOnFilterSkeletonTreeItem) override;
+	UE_API virtual void Build(FSkeletonTreeBuilderOutput& Output) override;
+	UE_API virtual void Filter(const FSkeletonTreeFilterArgs& InArgs, const TArray<TSharedPtr<class ISkeletonTreeItem>>& InItems, TArray<TSharedPtr<class ISkeletonTreeItem>>& OutFilteredItems) override;
+	UE_API virtual ESkeletonTreeFilterResult FilterItem(const FSkeletonTreeFilterArgs& InArgs, const TSharedPtr<class ISkeletonTreeItem>& InItem) override;
+	UE_API virtual bool IsShowingBones() const override;
+	UE_API virtual bool IsShowingSockets() const override;
+	UE_API virtual bool IsShowingAttachedAssets() const override;
 
 protected:
 	/** Add Bones */
-	void AddBones(FSkeletonTreeBuilderOutput& Output);
+	UE_API void AddBones(FSkeletonTreeBuilderOutput& Output);
 
 	/** Add Sockets */
-	void AddSockets(FSkeletonTreeBuilderOutput& Output);
+	UE_API void AddSockets(FSkeletonTreeBuilderOutput& Output);
 
-	void AddAttachedAssets(FSkeletonTreeBuilderOutput& Output);
+	UE_API void AddAttachedAssets(FSkeletonTreeBuilderOutput& Output);
 
-	void AddSocketsFromData(const TArray< USkeletalMeshSocket* >& SocketArray, ESocketParentType ParentType, FSkeletonTreeBuilderOutput& Output);
+	UE_API void AddSocketsFromData(const TArray< USkeletalMeshSocket* >& SocketArray, ESocketParentType ParentType, FSkeletonTreeBuilderOutput& Output);
 
-	void AddAttachedAssetContainer(const FPreviewAssetAttachContainer& AttachedObjects, FSkeletonTreeBuilderOutput& Output);
+	UE_API void AddAttachedAssetContainer(const FPreviewAssetAttachContainer& AttachedObjects, FSkeletonTreeBuilderOutput& Output);
 
-	void AddVirtualBones(FSkeletonTreeBuilderOutput& Output);
+	UE_API void AddVirtualBones(FSkeletonTreeBuilderOutput& Output);
 
 	/** Create an item for a bone */
-	TSharedRef<class ISkeletonTreeItem> CreateBoneTreeItem(const FName& InBoneName);
+	UE_API TSharedRef<class ISkeletonTreeItem> CreateBoneTreeItem(const FName& InBoneName);
 
 	/** Create an item for a socket */
-	TSharedRef<class ISkeletonTreeItem> CreateSocketTreeItem(USkeletalMeshSocket* InSocket, ESocketParentType ParentType, bool bIsCustomized);
+	UE_API TSharedRef<class ISkeletonTreeItem> CreateSocketTreeItem(USkeletalMeshSocket* InSocket, ESocketParentType ParentType, bool bIsCustomized);
 
 	/** Create an item for an attached asset */
-	TSharedRef<class ISkeletonTreeItem> CreateAttachedAssetTreeItem(UObject* InAsset, const FName& InAttachedTo);
+	UE_API TSharedRef<class ISkeletonTreeItem> CreateAttachedAssetTreeItem(UObject* InAsset, const FName& InAttachedTo);
 
 	/** Create an item for a virtual bone */
-	TSharedRef<class ISkeletonTreeItem> CreateVirtualBoneTreeItem(const FName& InBoneName);
+	UE_API TSharedRef<class ISkeletonTreeItem> CreateVirtualBoneTreeItem(const FName& InBoneName);
 
 	/** Helper function for filtering */
-	ESkeletonTreeFilterResult FilterRecursive(const FSkeletonTreeFilterArgs& InArgs, const TSharedPtr<ISkeletonTreeItem>& InItem, TArray<TSharedPtr<ISkeletonTreeItem>>& OutFilteredItems);
+	UE_API ESkeletonTreeFilterResult FilterRecursive(const FSkeletonTreeFilterArgs& InArgs, const TSharedPtr<ISkeletonTreeItem>& InItem, TArray<TSharedPtr<ISkeletonTreeItem>>& OutFilteredItems);
 
 protected:
 	/** Options for building */
@@ -116,3 +118,5 @@ protected:
 	/** The (optional) preview scene we will build against */
 	TWeakPtr<class IPersonaPreviewScene> PreviewScenePtr;
 };
+
+#undef UE_API
diff --git a/Engine/Source/Editor/SkeletonEditor/Public/SkeletonTreeItem.h b/Engine/Source/Editor/SkeletonEditor/Public/SkeletonTreeItem.h
--- a/Engine/Source/Editor/SkeletonEditor/Public/SkeletonTreeItem.h
+++ b/Engine/Source/Editor/SkeletonEditor/Public/SkeletonTreeItem.h
@@ -16,6 +16,8 @@
 #include "Widgets/Views/STableRow.h"
 #include "Widgets/Views/STableViewBase.h"
 
+#define UE_API SKELETONEDITOR_API
+
 class FDragDropEvent;
 class FText;
 class ITableRow;
@@ -26,7 +28,7 @@ struct FGeometry;
 struct FPointerEvent;
 
 /** Provides basic stub functionality for ISkeletonTreeItem-derived classes */
-class SKELETONEDITOR_API FSkeletonTreeItem : public ISkeletonTreeItem
+class FSkeletonTreeItem : public ISkeletonTreeItem
 {
 public:
 	SKELETON_TREE_ITEM_TYPE(FSkeletonTreeItem, ISkeletonTreeItem)
@@ -39,7 +41,7 @@ public:
 	{}
 
 	/** ISkeletonTreeItem interface */
-	virtual TSharedRef<ITableRow> MakeTreeRowWidget(const TSharedRef<STableViewBase>& InOwnerTable, const TAttribute<FText>& InFilterText) override;
+	UE_API virtual TSharedRef<ITableRow> MakeTreeRowWidget(const TSharedRef<STableViewBase>& InOwnerTable, const TAttribute<FText>& InFilterText) override;
 	virtual TSharedRef< SWidget > GenerateWidgetForDataColumn(const FName& DataColumnName, FIsSelected InIsSelected) override { return SNullWidget::NullWidget; }
 	virtual TSharedRef< SWidget > GenerateInlineEditWidget(const TAttribute<FText>& FilterText, FIsSelected InIsSelected) override { return SNullWidget::NullWidget; }
 	virtual bool HasInlineEditor() const override { return false; }
@@ -80,3 +82,5 @@ protected:
 	/** The current filter result */
 	ESkeletonTreeFilterResult FilterResult;
 };
+
+#undef UE_API
diff --git a/Engine/Source/Editor/SourceControlWindows/Public/SourceControlMenuContext.h b/Engine/Source/Editor/SourceControlWindows/Public/SourceControlMenuContext.h
--- a/Engine/Source/Editor/SourceControlWindows/Public/SourceControlMenuContext.h
+++ b/Engine/Source/Editor/SourceControlWindows/Public/SourceControlMenuContext.h
@@ -11,8 +11,8 @@
 /**
  * Source control window menu context providing information for menu extenders.
  */
-UCLASS()
-class SOURCECONTROLWINDOWS_API USourceControlMenuContext : public UObject
+UCLASS(MinimalAPI)
+class USourceControlMenuContext : public UObject
 {
 	GENERATED_BODY()
 public:
diff --git a/Engine/Source/Editor/SourceControlWindows/Public/SourceControlWindows.h b/Engine/Source/Editor/SourceControlWindows/Public/SourceControlWindows.h
--- a/Engine/Source/Editor/SourceControlWindows/Public/SourceControlWindows.h
+++ b/Engine/Source/Editor/SourceControlWindows/Public/SourceControlWindows.h
@@ -10,11 +10,13 @@
 #include "UObject/ObjectMacros.h"
 #include "SourceControlWindows.generated.h"
 
+#define UE_API SOURCECONTROLWINDOWS_API
+
 class SSourceControlHistoryWidget;
 
 /** Context object for SSourceControlHistoryWidget's right-click context menu to get info about which history widget is trying to generate the menu */
-UCLASS()
-class SOURCECONTROLWINDOWS_API USourceControlHistoryWidgetContext : public UObject
+UCLASS(MinimalAPI)
+class USourceControlHistoryWidgetContext : public UObject
 {
 	GENERATED_BODY()
 
@@ -38,10 +40,10 @@ private:
 
 
 /** Info supplied as argument to delegate FSourceControlWindowsOnCheckInComplete called by FSourceControlWindows::ChoosePackagesToCheckIn() and optional argument to PromptForCheckin(). */
-struct SOURCECONTROLWINDOWS_API FCheckinResultInfo
+struct FCheckinResultInfo
 {
 	/** Default Constructor */
-	FCheckinResultInfo();
+	UE_API FCheckinResultInfo();
 
 	/** Succeeded - if packages were selected and successfully checked in, Cancelled - if the user aborted the process, Failed - if an issue was encountered during the process */
 	ECommandResult::Type Result;
@@ -62,7 +64,7 @@ struct SOURCECONTROLWINDOWS_API FCheckinResultInfo
 /** Optional delegate called when FSourceControlWindows::ChoosePackagesToCheckIn() completes. */
 DECLARE_DELEGATE_OneParam(FSourceControlWindowsOnCheckInComplete, const FCheckinResultInfo&);
 
-class SOURCECONTROLWINDOWS_API FSourceControlWindows
+class FSourceControlWindows
 {
 public:
 	/**
@@ -71,10 +73,10 @@ public:
 	 * @param	OnCompleteDelegate	Delegate to call when this user-based operation is complete. Also see FCheckinResultInfo.
 	 * @return	true - if command successfully in progress and OnCompleteDelegate will be called when complete, false - if immediately unable to comply (such as source control not enabled)
 	 */
-	static bool ChoosePackagesToCheckIn(const FSourceControlWindowsOnCheckInComplete& OnCompleteDelegate = FSourceControlWindowsOnCheckInComplete());
+	static UE_API bool ChoosePackagesToCheckIn(const FSourceControlWindowsOnCheckInComplete& OnCompleteDelegate = FSourceControlWindowsOnCheckInComplete());
 
 	/** Determines whether we can choose packages to check in (we can't if an operation is already in progress) */
-	static bool CanChoosePackagesToCheckIn();
+	static UE_API bool CanChoosePackagesToCheckIn();
 
 	/** Determines if the Submit Content action should be visible or not */
 	UE_DEPRECATED(5.6, "ShouldChoosePackagesToCheckBeVisible is deprecated and will be removed")
@@ -85,17 +87,17 @@ public:
 	 * 
 	 * @return	true - if command completed successfully.
 	 */
-	static bool SyncLatest();
+	static UE_API bool SyncLatest();
 
 	/**
 	 * Saves all unsaved levels and assets and then - conditionally - performs an FSync operation to get the specified revision, and reloads the world
 	 *
 	 * @return	true - if command completed successfully.
 	 */
-	static bool SyncRevision(const FString& InRevision);
+	static UE_API bool SyncRevision(const FString& InRevision);
 
 	/** Determines whether we can sync to latest */
-	static bool CanSyncLatest();
+	static UE_API bool CanSyncLatest();
 
 	/**
 	 * Display check in dialog for the specified packages and get additional result information
@@ -108,7 +110,7 @@ public:
 	 * 
 	 * @return	true - if completed successfully in progress and OnCompleteDelegate will be called when complete, false - if immediately unable to comply (such as source control not enabled)
 	 */
-	static bool PromptForCheckin(FCheckinResultInfo& OutResultInfo, const TArray<FString>& InPackageNames, const TArray<FString>& InPendingDeletePaths = TArray<FString>(), const TArray<FString>& InConfigFiles = TArray<FString>(), bool bUseSourceControlStateCache = false);
+	static UE_API bool PromptForCheckin(FCheckinResultInfo& OutResultInfo, const TArray<FString>& InPackageNames, const TArray<FString>& InPendingDeletePaths = TArray<FString>(), const TArray<FString>& InConfigFiles = TArray<FString>(), bool bUseSourceControlStateCache = false);
 
 	/**
 	 * Display check in dialog for the specified packages
@@ -120,7 +122,7 @@ public:
 	 * 
 	 * @return	true - if completed successfully in progress and OnCompleteDelegate will be called when complete, false - if immediately unable to comply (such as source control not enabled)
 	 */
-	static bool PromptForCheckin(bool bUseSourceControlStateCache, const TArray<FString>& InPackageNames, const TArray<FString>& InPendingDeletePaths = TArray<FString>(), const TArray<FString>& InConfigFiles = TArray<FString>());
+	static UE_API bool PromptForCheckin(bool bUseSourceControlStateCache, const TArray<FString>& InPackageNames, const TArray<FString>& InPendingDeletePaths = TArray<FString>(), const TArray<FString>& InConfigFiles = TArray<FString>());
 
 	/**
 	 * Display checked out dialog for the specified filenames
@@ -142,14 +144,14 @@ public:
 		bool bShowColumnUserName = true;
 	};
 
-	static bool PromptForCheckedOut(bool bUseSourceControlStateCache, TArray<FString>& InFileNames, FCheckedOutSetupInfo& InSetupInfo);
+	static UE_API bool PromptForCheckedOut(bool bUseSourceControlStateCache, TArray<FString>& InFileNames, FCheckedOutSetupInfo& InSetupInfo);
 
 	/**
 	 * Display file revision history for the provided packages
 	 *
 	 * @param	InPackageNames	Names of packages to display file revision history for
 	 */
-	static void DisplayRevisionHistory( const TArray<FString>& InPackagesNames );
+	static UE_API void DisplayRevisionHistory( const TArray<FString>& InPackagesNames );
 
 	/**
 	 * Prompt the user with a revert files dialog, allowing them to specify which packages, if any, should be reverted.
@@ -159,12 +161,12 @@ public:
 	 *
 	 * @return	true if the files were reverted; false if the user canceled out of the dialog
 	 */
-	static bool PromptForRevert(const TArray<FString>& InPackageNames, bool bInReloadWorld = false );
+	static UE_API bool PromptForRevert(const TArray<FString>& InPackageNames, bool bInReloadWorld = false );
 
 	/**
 	 * Revert all locally modified files, and reload the world
 	 */
-	static bool RevertAllChangesAndReloadWorld();
+	static UE_API bool RevertAllChangesAndReloadWorld();
 
 	/**
 	 * Displays file diff against workspace version
@@ -173,7 +175,7 @@ public:
 	 * 
 	 * @return true if the diff could be performed
 	 */
-	static bool DiffAgainstWorkspace(const FString& InFileName);
+	static UE_API bool DiffAgainstWorkspace(const FString& InFileName);
 
 	/**
 	 * Displays file diff against shelved file
@@ -182,20 +184,21 @@ public:
 	 * 
 	 * @return true if the diff could be performed
 	 */
-	static bool DiffAgainstShelvedFile(const FSourceControlStateRef& InFileState);
+	static UE_API bool DiffAgainstShelvedFile(const FSourceControlStateRef& InFileState);
 
 protected:
 	/** Callback for ChoosePackagesToCheckIn(), continues to bring up UI once source control operations are complete */
-	static void ChoosePackagesToCheckInCallback(const FSourceControlOperationRef& InOperation, ECommandResult::Type InResult, FSourceControlWindowsOnCheckInComplete OnCompleteDelegate);
+	static UE_API void ChoosePackagesToCheckInCallback(const FSourceControlOperationRef& InOperation, ECommandResult::Type InResult, FSourceControlWindowsOnCheckInComplete OnCompleteDelegate);
 
 	/** Called when the user selection process has completed and we have packages to check in. */
-	static void ChoosePackagesToCheckInCompleted(const TArray<UPackage*>& LoadedPackages, const TArray<FString>& PackageNames, const TArray<FString>& ConfigFiles, FCheckinResultInfo& OutResultInfo);
+	static UE_API void ChoosePackagesToCheckInCompleted(const TArray<UPackage*>& LoadedPackages, const TArray<FString>& PackageNames, const TArray<FString>& ConfigFiles, FCheckinResultInfo& OutResultInfo);
 
 	/** Delegate called when the user has decided to cancel the check in process */
-	static void ChoosePackagesToCheckInCancelled(FSourceControlOperationRef InOperation);
+	static UE_API void ChoosePackagesToCheckInCancelled(FSourceControlOperationRef InOperation);
 
 private:
 	/** The notification in place while we choose packages to check in */
-	static TWeakPtr<class SNotificationItem> ChoosePackagesToCheckInNotification;
+	static UE_API TWeakPtr<class SNotificationItem> ChoosePackagesToCheckInNotification;
 };
 
+#undef UE_API
diff --git a/Engine/Source/Editor/SparseVolumeTexture/Public/OpenVDBImportOptions.h b/Engine/Source/Editor/SparseVolumeTexture/Public/OpenVDBImportOptions.h
--- a/Engine/Source/Editor/SparseVolumeTexture/Public/OpenVDBImportOptions.h
+++ b/Engine/Source/Editor/SparseVolumeTexture/Public/OpenVDBImportOptions.h
@@ -81,11 +81,11 @@ struct FOpenVDBPreviewData
 	FOpenVDBImportOptions DefaultImportOptions;
 };
 
-UCLASS()
-class SPARSEVOLUMETEXTURE_API UOpenVDBImportOptionsObject : public UObject
+UCLASS(MinimalAPI)
+class UOpenVDBImportOptionsObject : public UObject
 {
 	GENERATED_BODY()
 
 public:
 	FOpenVDBPreviewData PreviewData;
-};
\ No newline at end of file
+};
diff --git a/Engine/Source/Editor/StaticMeshEditor/Public/StaticMeshEditorSubsystem.h b/Engine/Source/Editor/StaticMeshEditor/Public/StaticMeshEditorSubsystem.h
--- a/Engine/Source/Editor/StaticMeshEditor/Public/StaticMeshEditorSubsystem.h
+++ b/Engine/Source/Editor/StaticMeshEditor/Public/StaticMeshEditorSubsystem.h
@@ -21,13 +21,13 @@ DECLARE_LOG_CATEGORY_EXTERN(LogStaticMeshEditorSubsystem, Log, All);
 * UStaticMeshEditorSubsystem
 * Subsystem for exposing static mesh functionality to scripts
 */
-UCLASS()
-class STATICMESHEDITOR_API UStaticMeshEditorSubsystem : public UEditorSubsystem
+UCLASS(MinimalAPI)
+class UStaticMeshEditorSubsystem : public UEditorSubsystem
 {
 	GENERATED_BODY()
 
 public:
-	UStaticMeshEditorSubsystem();
+	STATICMESHEDITOR_API UStaticMeshEditorSubsystem();
 
 	/**
 	 * Remove then add LODs on a static mesh.
@@ -42,7 +42,7 @@ public:
 	 * No action will be performed if ReductionOptions.ReductionSettings is empty
 	 */
 	UFUNCTION(BlueprintCallable, Category = "Static Mesh Utilities")
-	int32 SetLodsWithNotification(UStaticMesh* StaticMesh, const FStaticMeshReductionOptions& ReductionOptions, bool bApplyChanges);
+	STATICMESHEDITOR_API int32 SetLodsWithNotification(UStaticMesh* StaticMesh, const FStaticMeshReductionOptions& ReductionOptions, bool bApplyChanges);
 
 	/**
 	 * Same as SetLodsWithNotification but changes are applied.
@@ -60,7 +60,7 @@ public:
 	 * @param OutReductionOptions - The reduction settings where we copy the reduction options.
 	 */
 	UFUNCTION(BlueprintCallable, Category = "Editor Scripting | StaticMesh")
-	void GetLodReductionSettings(const UStaticMesh* StaticMesh, const int32 LodIndex, FMeshReductionSettings& OutReductionOptions);
+	STATICMESHEDITOR_API void GetLodReductionSettings(const UStaticMesh* StaticMesh, const int32 LodIndex, FMeshReductionSettings& OutReductionOptions);
 
 	/**
 	 * Set the LOD reduction for the specified LOD index.
@@ -69,7 +69,7 @@ public:
 	 * @param ReductionOptions - The reduction settings we want to apply to the LOD.
 	 */
 	UFUNCTION(BlueprintCallable, Category = "Editor Scripting | StaticMesh")
-	void SetLodReductionSettings(UStaticMesh* StaticMesh, const int32 LodIndex, const FMeshReductionSettings& ReductionOptions);
+	STATICMESHEDITOR_API void SetLodReductionSettings(UStaticMesh* StaticMesh, const int32 LodIndex, const FMeshReductionSettings& ReductionOptions);
 
 	/**
 	 * Copy the build options with the specified LOD build settings.
@@ -78,7 +78,7 @@ public:
 	 * @param OutBuildOptions - The build settings where we copy the build options.
 	 */
 	UFUNCTION(BlueprintCallable, Category = "Editor Scripting | StaticMesh")
-	void GetLodBuildSettings(const UStaticMesh* StaticMesh, const int32 LodIndex, FMeshBuildSettings& OutBuildOptions);
+	STATICMESHEDITOR_API void GetLodBuildSettings(const UStaticMesh* StaticMesh, const int32 LodIndex, FMeshBuildSettings& OutBuildOptions);
 
 	/**
 	 * Set the LOD build options for the specified LOD index.
@@ -87,7 +87,7 @@ public:
 	 * @param BuildOptions - The build settings we want to apply to the LOD.
 	 */
 	UFUNCTION(BlueprintCallable, Category = "Editor Scripting | StaticMesh")
-	void SetLodBuildSettings(UStaticMesh* StaticMesh, const int32 LodIndex, const FMeshBuildSettings& BuildOptions);
+	STATICMESHEDITOR_API void SetLodBuildSettings(UStaticMesh* StaticMesh, const int32 LodIndex, const FMeshBuildSettings& BuildOptions);
 
 	/**
 	 * Get the LODGroup for the specified static mesh
@@ -95,7 +95,7 @@ public:
 	 * @return LODGroup
 	 */
 	UFUNCTION(BlueprintCallable,Category = "Editor Scripting | StaticMesh")
-	FName GetLODGroup(const UStaticMesh* StaticMesh);
+	STATICMESHEDITOR_API FName GetLODGroup(const UStaticMesh* StaticMesh);
 
 	/**
 	 * Set the LODGroup for the specified static mesh
@@ -105,7 +105,7 @@ public:
 	 * @return Success
 	 */
 	UFUNCTION(BlueprintCallable, Category = "Editor Scripting | StaticMesh")
-	bool SetLODGroup(UStaticMesh* StaticMesh, FName LODGroup, bool bRebuildImmediately=true);
+	STATICMESHEDITOR_API bool SetLODGroup(UStaticMesh* StaticMesh, FName LODGroup, bool bRebuildImmediately=true);
 	
 	/**
 	 * Import or re-import a LOD into the specified base mesh. If the LOD do not exist it will import it and add it to the base static mesh. If the LOD already exist it will re-import the specified LOD.
@@ -117,7 +117,7 @@ public:
 	 * @return the index of the LOD that was imported or re-imported. Will return INDEX_NONE if anything goes bad.
 	 */
 	UFUNCTION(BlueprintCallable, Category = "Editor Scripting | StaticMesh")
-	int32 ImportLOD(UStaticMesh* BaseStaticMesh, const int32 LODIndex, const FString& SourceFilename);
+	STATICMESHEDITOR_API int32 ImportLOD(UStaticMesh* BaseStaticMesh, const int32 LODIndex, const FString& SourceFilename);
 
 	/**
 	 * Re-import all the custom LODs present in the specified static mesh.
@@ -127,7 +127,7 @@ public:
 	 * @return true if re-import all LODs works, false otherwise see log for explanation.
 	 */
 	UFUNCTION(BlueprintCallable, Category = "Editor Scripting | StaticMesh", meta=(ScriptName="ReimportAllCustomLods;ReimportAllCustomLODs"))
-	bool ReimportAllCustomLODs(UStaticMesh* StaticMesh);
+	STATICMESHEDITOR_API bool ReimportAllCustomLODs(UStaticMesh* StaticMesh);
 	/**
 	 * Adds or create a LOD at DestinationLodIndex using the geometry from SourceStaticMesh SourceLodIndex
 	 * @param	DestinationStaticMesh		The static mesh to set the LOD in.
@@ -140,7 +140,7 @@ public:
 	 *			A negative value indicates that the LOD was not set. See log for explanation.
 	 */
 	UFUNCTION(BlueprintCallable, Category = "Static Mesh Utilities")
-	int32 SetLodFromStaticMesh(UStaticMesh* DestinationStaticMesh, int32 DestinationLodIndex, UStaticMesh* SourceStaticMesh, int32 SourceLodIndex, bool bReuseExistingMaterialSlots);
+	STATICMESHEDITOR_API int32 SetLodFromStaticMesh(UStaticMesh* DestinationStaticMesh, int32 DestinationLodIndex, UStaticMesh* SourceStaticMesh, int32 SourceLodIndex, bool bReuseExistingMaterialSlots);
 
 	/**
 	 * Get number of LODs present on a static mesh.
@@ -149,7 +149,7 @@ public:
 	 * An negative value indicates that the command could not be executed. See log for explanation.
 	 */
 	UFUNCTION(BlueprintCallable, Category = "Static Mesh Utilities")
-	int32 GetLodCount(const UStaticMesh* StaticMesh);
+	STATICMESHEDITOR_API int32 GetLodCount(const UStaticMesh* StaticMesh);
 
 	/**
 	 * Remove LODs on a static mesh except LOD 0.
@@ -157,7 +157,7 @@ public:
 	 * @return A boolean indicating if the removal was successful, true, or not.
 	 */
 	UFUNCTION(BlueprintCallable, Category = "Static Mesh Utilities")
-	bool RemoveLods(UStaticMesh* StaticMesh);
+	STATICMESHEDITOR_API bool RemoveLods(UStaticMesh* StaticMesh);
 
 	/**
 	 * Get an array of LOD screen sizes for evaluation.
@@ -165,7 +165,7 @@ public:
 	 * @return array of LOD screen sizes.
 	 */
 	UFUNCTION(BlueprintPure, Category = "Static Mesh Utilities")
-	TArray<float> GetLodScreenSizes(const UStaticMesh* StaticMesh);
+	STATICMESHEDITOR_API TArray<float> GetLodScreenSizes(const UStaticMesh* StaticMesh);
 
 	/**
 	 * Set LOD screen sizes.
@@ -174,7 +174,7 @@ public:
 	 * @return A boolean indicating if setting the screen sizes was successful.
 	 */
 	UFUNCTION(BlueprintCallable, Category = "Static Mesh Utilities")
-	bool SetLodScreenSizes(UStaticMesh* StaticMesh, const TArray<float>& ScreenSizes);
+	STATICMESHEDITOR_API bool SetLodScreenSizes(UStaticMesh* StaticMesh, const TArray<float>& ScreenSizes);
 
 public:
 
@@ -189,7 +189,7 @@ public:
 	 * @return true if the nanite hi-res mesh import/re-import is successful, false otherwise.
 	 */
 	UFUNCTION(BlueprintCallable, Category = "Editor Scripting | StaticMesh")
-	bool ImportNaniteHiResMesh(UStaticMesh* StaticMesh, const FString& SourceFilename, bool bShowDialogWhenFileMissing);
+	STATICMESHEDITOR_API bool ImportNaniteHiResMesh(UStaticMesh* StaticMesh, const FString& SourceFilename, bool bShowDialogWhenFileMissing);
 
 	/**
 	* Get the Nanite hi-res source filename
@@ -197,7 +197,7 @@ public:
 	* @return the source filename use by the static mesh nanite hi-res.
 	*/
 	UFUNCTION(BlueprintCallable, Category = "Editor Scripting | StaticMesh")
-	FString GetNaniteSourceFilename(const UStaticMesh* StaticMesh);
+	STATICMESHEDITOR_API FString GetNaniteSourceFilename(const UStaticMesh* StaticMesh);
 
 	/**
 	* Update the Nanite hi-res source filename.
@@ -206,7 +206,7 @@ public:
 	* @return true if the source filename was updated (cleared or replace by an existing filename), false otherwise.
 	*/
 	UFUNCTION(BlueprintCallable, Category = "Editor Scripting | StaticMesh")
-	bool UpdateNaniteSourceFilename(UStaticMesh* StaticMesh, const FString& NewSourceFilename);
+	STATICMESHEDITOR_API bool UpdateNaniteSourceFilename(UStaticMesh* StaticMesh, const FString& NewSourceFilename);
 
 	/**
 	* Get the Nanite Settings for the mesh
@@ -214,7 +214,7 @@ public:
 	* @return FMeshNaniteSettings struct for the given static mesh
 	*/
 	UFUNCTION(BlueprintPure, Category = "Static Mesh Utilities")
-	FMeshNaniteSettings GetNaniteSettings(const UStaticMesh* StaticMesh);
+	STATICMESHEDITOR_API FMeshNaniteSettings GetNaniteSettings(const UStaticMesh* StaticMesh);
 
 	/**
 	* Get the Nanite Settings for the mesh
@@ -223,7 +223,7 @@ public:
 	* @param bApplyChanges     Indicates if changes must be applied or not.
 	*/
 	UFUNCTION(BlueprintCallable, Category = "Static Mesh Utilities")
-	void SetNaniteSettings(UStaticMesh* StaticMesh, FMeshNaniteSettings NaniteSettings, bool bApplyChanges=true);
+	STATICMESHEDITOR_API void SetNaniteSettings(UStaticMesh* StaticMesh, FMeshNaniteSettings NaniteSettings, bool bApplyChanges=true);
 
 
 	
@@ -238,7 +238,7 @@ public:
 	 * A negative value indicates the addition failed.
 	 */
 	UFUNCTION(BlueprintCallable, Category = "Static Mesh Utilities")
-	int32 AddSimpleCollisionsWithNotification(UStaticMesh* StaticMesh, const EScriptCollisionShapeType ShapeType, bool bApplyChanges);
+	STATICMESHEDITOR_API int32 AddSimpleCollisionsWithNotification(UStaticMesh* StaticMesh, const EScriptCollisionShapeType ShapeType, bool bApplyChanges);
 
 	/**
 	 * Same as AddSimpleCollisionsWithNotification but changes are automatically applied.
@@ -256,7 +256,7 @@ public:
 	 * An negative value indicates that the command could not be executed. See log for explanation.
 	 */
 	UFUNCTION(BlueprintCallable, Category = "Static Mesh Utilities")
-	int32 GetSimpleCollisionCount(const UStaticMesh* StaticMesh);
+	STATICMESHEDITOR_API int32 GetSimpleCollisionCount(const UStaticMesh* StaticMesh);
 
 	/**
 	 * Get the Collision Trace behavior of a static mesh
@@ -264,7 +264,7 @@ public:
 	 * @return the Collision Trace behavior.
 	 */
 	UFUNCTION(BlueprintPure, Category = "Static Mesh Utilities")
-	TEnumAsByte<ECollisionTraceFlag> GetCollisionComplexity(const UStaticMesh* StaticMesh);
+	STATICMESHEDITOR_API TEnumAsByte<ECollisionTraceFlag> GetCollisionComplexity(const UStaticMesh* StaticMesh);
 
 	/**
 	 * Get number of convex collisions present on a static mesh.
@@ -273,7 +273,7 @@ public:
 	 * An negative value indicates that the command could not be executed. See log for explanation.
 	 */
 	UFUNCTION(BlueprintCallable, Category = "Static Mesh Utilities")
-	int32 GetConvexCollisionCount(const UStaticMesh* StaticMesh);
+	STATICMESHEDITOR_API int32 GetConvexCollisionCount(const UStaticMesh* StaticMesh);
 
 	/**
 	 * Add a convex collision to a static mesh.
@@ -287,7 +287,7 @@ public:
 	 * @return A boolean indicating if the addition was successful or not.
 	 */
 	UFUNCTION(BlueprintCallable, Category = "Static Mesh Utilities")
-	bool SetConvexDecompositionCollisionsWithNotification(UStaticMesh* StaticMesh, int32 HullCount, int32 MaxHullVerts, int32 HullPrecision, bool bApplyChanges);
+	STATICMESHEDITOR_API bool SetConvexDecompositionCollisionsWithNotification(UStaticMesh* StaticMesh, int32 HullCount, int32 MaxHullVerts, int32 HullPrecision, bool bApplyChanges);
 
 	/**
 	 * Compute convex collisions for a set of static meshes.
@@ -301,7 +301,7 @@ public:
 	 * @return A boolean indicating if the addition was successful or not.
 	 */
 	UFUNCTION(BlueprintCallable, Category = "Static Mesh Utilities")
-	bool BulkSetConvexDecompositionCollisionsWithNotification(const TArray<UStaticMesh*>& StaticMeshes, int32 HullCount, int32 MaxHullVerts, int32 HullPrecision, bool bApplyChanges);
+	STATICMESHEDITOR_API bool BulkSetConvexDecompositionCollisionsWithNotification(const TArray<UStaticMesh*>& StaticMeshes, int32 HullCount, int32 MaxHullVerts, int32 HullPrecision, bool bApplyChanges);
 
 	/**
 	 * Same as SetConvexDecompositionCollisionsWithNotification but changes are automatically applied.
@@ -329,7 +329,7 @@ public:
 	 * @return A boolean indicating if the removal was successful or not.
 	 */
 	UFUNCTION(BlueprintCallable, Category = "Static Mesh Utilities")
-	bool RemoveCollisionsWithNotification(UStaticMesh* StaticMesh, bool bApplyChanges);
+	STATICMESHEDITOR_API bool RemoveCollisionsWithNotification(UStaticMesh* StaticMesh, bool bApplyChanges);
 
 	/**
 	 * Same as RemoveCollisionsWithNotification but changes are applied.
@@ -344,7 +344,7 @@ private:
 	/**
 	* Checks if section and lod indices are valid indices.
 	*/
-	static bool CanModifyStaticMeshSection(const UStaticMesh* StaticMesh, int32 LODIndex, int32 SectionIndex, const FString& LogFunctionName);
+	static STATICMESHEDITOR_API bool CanModifyStaticMeshSection(const UStaticMesh* StaticMesh, int32 LODIndex, int32 SectionIndex, const FString& LogFunctionName);
 
 public:
 	/**
@@ -355,7 +355,7 @@ public:
 	 * @param	SectionIndex		Index of the StaticMesh Section.
 	 */
 	UFUNCTION(BlueprintCallable, Category = "Static Mesh Utilities")
-	void EnableSectionCastShadow(UStaticMesh* StaticMesh, bool bCastShadow, int32 LODIndex, int32 SectionIndex);
+	STATICMESHEDITOR_API void EnableSectionCastShadow(UStaticMesh* StaticMesh, bool bCastShadow, int32 LODIndex, int32 SectionIndex);
 
 	/**
 	 * Checks if a specific LOD mesh section can cast shadows.
@@ -364,7 +364,7 @@ public:
 	 * @param	SectionIndex		Index of the StaticMesh Section.
 	 */
 	UFUNCTION(BlueprintCallable, Category = "Static Mesh Utilities")
-	bool IsSectionCastShadowEnabled(const UStaticMesh* StaticMesh, int32 LODIndex, int32 SectionIndex);
+	STATICMESHEDITOR_API bool IsSectionCastShadowEnabled(const UStaticMesh* StaticMesh, int32 LODIndex, int32 SectionIndex);
 
 	/**
 	 * Enables/disables mesh section collision for a specific LOD.
@@ -374,7 +374,7 @@ public:
 	 * @param	SectionIndex		Index of the StaticMesh Section.
 	 */
 	UFUNCTION(BlueprintCallable, Category = "Static Mesh Utilities")
-	void EnableSectionCollision(UStaticMesh* StaticMesh, bool bCollisionEnabled, int32 LODIndex, int32 SectionIndex);
+	STATICMESHEDITOR_API void EnableSectionCollision(UStaticMesh* StaticMesh, bool bCollisionEnabled, int32 LODIndex, int32 SectionIndex);
 
 	/**
 	 * Checks if a specific LOD mesh section has collision.
@@ -384,7 +384,7 @@ public:
 	 * @return True if the collision is enabled for the specified LOD of the StaticMesh section.
 	 */
 	UFUNCTION(BlueprintPure, Category = "Static Mesh Utilities")
-	bool IsSectionCollisionEnabled(const UStaticMesh* StaticMesh, int32 LODIndex, int32 SectionIndex);
+	STATICMESHEDITOR_API bool IsSectionCollisionEnabled(const UStaticMesh* StaticMesh, int32 LODIndex, int32 SectionIndex);
 
 	/**
 	 * Enables/disables mesh section's visibility in ray tracing for a specific LOD.
@@ -394,7 +394,7 @@ public:
 	 * @param	SectionIndex			Index of the StaticMesh Section.
 	 */
 	UFUNCTION(BlueprintCallable, Category = "Static Mesh Utilities")
-	void EnableSectionVisibleInRayTracing(UStaticMesh* StaticMesh, bool bVisibleInRayTracing, int32 LODIndex, int32 SectionIndex);
+	STATICMESHEDITOR_API void EnableSectionVisibleInRayTracing(UStaticMesh* StaticMesh, bool bVisibleInRayTracing, int32 LODIndex, int32 SectionIndex);
 
 	/**
 	 * Checks if a specific LOD mesh section is visible in ray tracing.
@@ -403,7 +403,7 @@ public:
 	 * @param	SectionIndex		Index of the StaticMesh Section.
 	 */
 	UFUNCTION(BlueprintCallable, Category = "Static Mesh Utilities")
-	bool IsSectionVisibleInRayTracingEnabled(const UStaticMesh* StaticMesh, int32 LODIndex, int32 SectionIndex);
+	STATICMESHEDITOR_API bool IsSectionVisibleInRayTracingEnabled(const UStaticMesh* StaticMesh, int32 LODIndex, int32 SectionIndex);
 
 	/**
 	 * Enables/disables mesh section affects distance field lighting for a specific LOD.
@@ -413,7 +413,7 @@ public:
 	 * @param	SectionIndex					Index of the StaticMesh Section.
 	 */
 	UFUNCTION(BlueprintCallable, Category = "Static Mesh Utilities")
-	void EnableSectionAffectDistanceFieldLighting(UStaticMesh* StaticMesh, bool bAffectDistanceFieldLighting, int32 LODIndex, int32 SectionIndex);
+	STATICMESHEDITOR_API void EnableSectionAffectDistanceFieldLighting(UStaticMesh* StaticMesh, bool bAffectDistanceFieldLighting, int32 LODIndex, int32 SectionIndex);
 
 	/**
 	 * Checks if a specific LOD mesh section affects distance field lighting.
@@ -422,7 +422,7 @@ public:
 	 * @param	SectionIndex		Index of the StaticMesh Section.
 	 */
 	UFUNCTION(BlueprintCallable, Category = "Static Mesh Utilities")
-	bool IsSectionAffectDistanceFieldLightingEnabled(const UStaticMesh* StaticMesh, int32 LODIndex, int32 SectionIndex);
+	STATICMESHEDITOR_API bool IsSectionAffectDistanceFieldLightingEnabled(const UStaticMesh* StaticMesh, int32 LODIndex, int32 SectionIndex);
 
 
 	/**
@@ -433,7 +433,7 @@ public:
 	 * @param	SectionIndex		Index of the StaticMesh Section.
 	 */
 	UFUNCTION(BlueprintCallable, Category = "Static Mesh Utilities")
-	void EnableSectionForceOpaque(UStaticMesh* StaticMesh, bool bForceOpaque, int32 LODIndex, int32 SectionIndex);
+	STATICMESHEDITOR_API void EnableSectionForceOpaque(UStaticMesh* StaticMesh, bool bForceOpaque, int32 LODIndex, int32 SectionIndex);
 
 	/**
 	 * Checks if a specific LOD mesh section is forced to be opaque.
@@ -442,7 +442,7 @@ public:
 	 * @param	SectionIndex		Index of the StaticMesh Section.
 	 */
 	UFUNCTION(BlueprintCallable, Category = "Static Mesh Utilities")
-	bool IsSectionForceOpaqueEnabled(const UStaticMesh* StaticMesh, int32 LODIndex, int32 SectionIndex);
+	STATICMESHEDITOR_API bool IsSectionForceOpaqueEnabled(const UStaticMesh* StaticMesh, int32 LODIndex, int32 SectionIndex);
 
 public:
 	/**
@@ -453,7 +453,7 @@ public:
 	* @param	SectionIndex		Index of the StaticMesh Section.
 	*/
 	UFUNCTION(BlueprintCallable, Category = "Static Mesh Utilities")
-	void SetLODMaterialSlot(UStaticMesh* StaticMesh, int32 MaterialSlotIndex, int32 LODIndex, int32 SectionIndex);
+	STATICMESHEDITOR_API void SetLODMaterialSlot(UStaticMesh* StaticMesh, int32 MaterialSlotIndex, int32 LODIndex, int32 SectionIndex);
 
 	/**
 	 * Gets the material slot used for a specific LOD section.
@@ -463,31 +463,31 @@ public:
 	 * @return  MaterialSlotIndex	Index of the material slot used by the section or INDEX_NONE in case of error.
 	 */
 	UFUNCTION(BlueprintCallable, Category = "Static Mesh Utilities")
-	int32 GetLODMaterialSlot(const UStaticMesh* StaticMesh, int32 LODIndex, int32 SectionIndex);
+	STATICMESHEDITOR_API int32 GetLODMaterialSlot(const UStaticMesh* StaticMesh, int32 LODIndex, int32 SectionIndex);
 
 	/** Check whether a static mesh has vertex colors */
 	UFUNCTION(BlueprintPure, Category = "Static Mesh Utilities")
-	bool HasVertexColors(const UStaticMesh* StaticMesh);
+	STATICMESHEDITOR_API bool HasVertexColors(const UStaticMesh* StaticMesh);
 
 	/** Check whether a static mesh component has vertex colors */
 	UFUNCTION(BlueprintPure, Category = "Static Mesh Utilities")
-	bool HasInstanceVertexColors(const UStaticMeshComponent* StaticMeshComponent);
+	STATICMESHEDITOR_API bool HasInstanceVertexColors(const UStaticMeshComponent* StaticMeshComponent);
 
 	/** Set Generate Lightmap UVs for StaticMesh */
 	UFUNCTION(BlueprintCallable, Category = "Static Mesh Utilities", meta = (ScriptName = "SetGenerateLightmapUv"))
-	bool SetGenerateLightmapUVs(UStaticMesh* StaticMesh, bool bGenerateLightmapUVs);
+	STATICMESHEDITOR_API bool SetGenerateLightmapUVs(UStaticMesh* StaticMesh, bool bGenerateLightmapUVs);
 
 	/** Get number of StaticMesh verts for an LOD */
 	UFUNCTION(BlueprintPure, Category = "Static Mesh Utilities")
-	int32 GetNumberVerts(const UStaticMesh* StaticMesh, int32 LODIndex);
+	STATICMESHEDITOR_API int32 GetNumberVerts(const UStaticMesh* StaticMesh, int32 LODIndex);
 
 	/** Get number of Materials for a StaticMesh */
 	UFUNCTION(BlueprintPure, Category = "Static Mesh Utilities")
-	int32 GetNumberMaterials(const UStaticMesh* StaticMesh);
+	STATICMESHEDITOR_API int32 GetNumberMaterials(const UStaticMesh* StaticMesh);
 
 	/** Sets StaticMeshFlag bAllowCPUAccess  */
 	UFUNCTION(BlueprintCallable, Category = "Static Mesh Utilities")
-	void SetAllowCPUAccess(UStaticMesh* StaticMesh, bool bAllowCPUAccess);
+	STATICMESHEDITOR_API void SetAllowCPUAccess(UStaticMesh* StaticMesh, bool bAllowCPUAccess);
 
 public:
 
@@ -498,7 +498,7 @@ public:
 	 * @return the number of UV channels.
 	 */
 	UFUNCTION(BlueprintCallable, Category = "Static Mesh Utilities")
-	int32 GetNumUVChannels(const UStaticMesh* StaticMesh, int32 LODIndex);
+	STATICMESHEDITOR_API int32 GetNumUVChannels(const UStaticMesh* StaticMesh, int32 LODIndex);
 
 	/**
 	 * Adds an empty UV channel at the end of the existing channels on the given LOD of a StaticMesh.
@@ -507,7 +507,7 @@ public:
 	 * @return true if a UV channel was added.
 	 */
 	UFUNCTION(BlueprintCallable, Category = "Static Mesh Utilities")
-	bool AddUVChannel(UStaticMesh* StaticMesh, int32 LODIndex);
+	STATICMESHEDITOR_API bool AddUVChannel(UStaticMesh* StaticMesh, int32 LODIndex);
 
 	/**
 	 * Inserts an empty UV channel at the specified channel index on the given LOD of a StaticMesh.
@@ -517,7 +517,7 @@ public:
 	 * @return true if a UV channel was added.
 	 */
 	UFUNCTION(BlueprintCallable, Category = "Static Mesh Utilities")
-	bool InsertUVChannel(UStaticMesh* StaticMesh, int32 LODIndex, int32 UVChannelIndex);
+	STATICMESHEDITOR_API bool InsertUVChannel(UStaticMesh* StaticMesh, int32 LODIndex, int32 UVChannelIndex);
 
 	/**
 	 * Removes the UV channel at the specified channel index on the given LOD of a StaticMesh.
@@ -527,7 +527,7 @@ public:
 	 * @return true if the UV channel was removed.
 	 */
 	UFUNCTION(BlueprintCallable, Category = "Static Mesh Utilities")
-	bool RemoveUVChannel(UStaticMesh* StaticMesh, int32 LODIndex, int32 UVChannelIndex);
+	STATICMESHEDITOR_API bool RemoveUVChannel(UStaticMesh* StaticMesh, int32 LODIndex, int32 UVChannelIndex);
 
 	/**
 	 * Generates planar UV mapping in the specified UV channel on the given LOD of a StaticMesh.
@@ -540,7 +540,7 @@ public:
 	 * @return true if the UV mapping was generated.
 	 */
 	UFUNCTION(BlueprintCallable, Category = "Static Mesh Utilities")
-	bool GeneratePlanarUVChannel(UStaticMesh* StaticMesh, int32 LODIndex, int32 UVChannelIndex, const FVector& Position, const FRotator& Orientation, const FVector2D& Tiling);
+	STATICMESHEDITOR_API bool GeneratePlanarUVChannel(UStaticMesh* StaticMesh, int32 LODIndex, int32 UVChannelIndex, const FVector& Position, const FRotator& Orientation, const FVector2D& Tiling);
 
 	/**
 	 * Generates cylindrical UV mapping in the specified UV channel on the given LOD of a StaticMesh.
@@ -553,7 +553,7 @@ public:
 	 * @return true if the UV mapping was generated.
 	 */
 	UFUNCTION(BlueprintCallable, Category = "Static Mesh Utilities")
-	bool GenerateCylindricalUVChannel(UStaticMesh* StaticMesh, int32 LODIndex, int32 UVChannelIndex, const FVector& Position, const FRotator& Orientation, const FVector2D& Tiling);
+	STATICMESHEDITOR_API bool GenerateCylindricalUVChannel(UStaticMesh* StaticMesh, int32 LODIndex, int32 UVChannelIndex, const FVector& Position, const FRotator& Orientation, const FVector2D& Tiling);
 
 	/**
 	 * Generates box UV mapping in the specified UV channel on the given LOD of a StaticMesh.
@@ -566,7 +566,7 @@ public:
 	 * @return true if the UV mapping was generated.
 	 */
 	UFUNCTION(BlueprintCallable, Category = "Static Mesh Utilities")
-	bool GenerateBoxUVChannel(UStaticMesh* StaticMesh, int32 LODIndex, int32 UVChannelIndex, const FVector& Position, const FRotator& Orientation, const FVector& Size);
+	STATICMESHEDITOR_API bool GenerateBoxUVChannel(UStaticMesh* StaticMesh, int32 LODIndex, int32 UVChannelIndex, const FVector& Position, const FRotator& Orientation, const FVector& Size);
 
 	/**
 	 * Find the references of the material MaterialToReplaced on all the MeshComponents provided and replace it by NewMaterial.
@@ -575,7 +575,7 @@ public:
 	 * @param	NewMaterial				Material to replace MaterialToBeReplaced by.
 	 */
 	UFUNCTION(BlueprintCallable, Category = "Static Mesh Utilities")
-	void ReplaceMeshComponentsMaterials(const TArray<class UMeshComponent*>& MeshComponents, class UMaterialInterface* MaterialToBeReplaced, class UMaterialInterface* NewMaterial);
+	STATICMESHEDITOR_API void ReplaceMeshComponentsMaterials(const TArray<class UMeshComponent*>& MeshComponents, class UMaterialInterface* MaterialToBeReplaced, class UMaterialInterface* NewMaterial);
 
 	/**
 	 * Find the references of the material MaterialToReplaced on all the MeshComponents of all the Actors provided and replace it by NewMaterial.
@@ -584,7 +584,7 @@ public:
 	 * @param	NewMaterial				Material to replace MaterialToBeReplaced by.
 	 */
 	UFUNCTION(BlueprintCallable, Category = "Static Mesh Utilities")
-	void ReplaceMeshComponentsMaterialsOnActors(const TArray<class AActor*>& Actors, class UMaterialInterface* MaterialToBeReplaced, class UMaterialInterface* NewMaterial);
+	STATICMESHEDITOR_API void ReplaceMeshComponentsMaterialsOnActors(const TArray<class AActor*>& Actors, class UMaterialInterface* MaterialToBeReplaced, class UMaterialInterface* NewMaterial);
 
 	/**
 	 * Find the references of the mesh MeshToBeReplaced on all the MeshComponents provided and replace it by NewMesh.
@@ -594,7 +594,7 @@ public:
 	 * @param	NewMesh					Mesh to replace MeshToBeReplaced by.
 	 */
 	UFUNCTION(BlueprintCallable, Category = "Static Mesh Utilities")
-	void ReplaceMeshComponentsMeshes(const TArray<class UStaticMeshComponent*>& MeshComponents, class UStaticMesh* MeshToBeReplaced, class UStaticMesh* NewMesh);
+	STATICMESHEDITOR_API void ReplaceMeshComponentsMeshes(const TArray<class UStaticMeshComponent*>& MeshComponents, class UStaticMesh* MeshToBeReplaced, class UStaticMesh* NewMesh);
 
 	/**
 	 * Find the references of the mesh MeshToBeReplaced on all the MeshComponents of all the Actors provided and replace it by NewMesh.
@@ -603,7 +603,7 @@ public:
 	 * @param	NewMesh					Mesh to replace MeshToBeReplaced by.
 	 */
 	UFUNCTION(BlueprintCallable, Category = "Static Mesh Utilities")
-	void ReplaceMeshComponentsMeshesOnActors(const TArray<class AActor*>& Actors, class UStaticMesh* MeshToBeReplaced, class UStaticMesh* NewMesh);
+	STATICMESHEDITOR_API void ReplaceMeshComponentsMeshesOnActors(const TArray<class AActor*>& Actors, class UStaticMesh* MeshToBeReplaced, class UStaticMesh* NewMesh);
 
 		/**
 	 * Create a new Actor in the level that contains a duplicate of all the Actors Static Meshes Component.
@@ -614,7 +614,7 @@ public:
 	 * @return The new created actor.
 	 */
 	UFUNCTION(BlueprintCallable, Category = "Editor Scripting | Dataprep")
-	class AActor* JoinStaticMeshActors(const TArray<class AStaticMeshActor*>& ActorsToJoin, const FJoinStaticMeshActorsOptions& JoinOptions);
+	STATICMESHEDITOR_API class AActor* JoinStaticMeshActors(const TArray<class AStaticMeshActor*>& ActorsToJoin, const FJoinStaticMeshActorsOptions& JoinOptions);
 
 	/**
 	 * Merge the meshes into a unique mesh with the provided StaticMeshActors. There are multiple options on how to merge the meshes and their materials.
@@ -626,7 +626,7 @@ public:
 	 * @return	if the operation is successful.
 	 */
 	UFUNCTION(BlueprintCallable, Category = "Editor Scripting | Dataprep")
-	bool MergeStaticMeshActors(const TArray<class AStaticMeshActor*>& ActorsToMerge, const FMergeStaticMeshActorsOptions& MergeOptions, class AStaticMeshActor*& OutMergedActor);
+	STATICMESHEDITOR_API bool MergeStaticMeshActors(const TArray<class AStaticMeshActor*>& ActorsToMerge, const FMergeStaticMeshActorsOptions& MergeOptions, class AStaticMeshActor*& OutMergedActor);
 
 	/**
 	 * Build a proxy mesh actor that can replace a set of mesh actors.
@@ -636,9 +636,9 @@ public:
 	 * @return  Success of the proxy creation
 	 */
 	UFUNCTION(BlueprintCallable, Category = "Editor Scripting | Dataprep")
-	bool CreateProxyMeshActor(const TArray<class AStaticMeshActor*>& ActorsToMerge, const FCreateProxyMeshActorOptions& MergeOptions, class AStaticMeshActor*& OutMergedActor);
+	STATICMESHEDITOR_API bool CreateProxyMeshActor(const TArray<class AStaticMeshActor*>& ActorsToMerge, const FCreateProxyMeshActorOptions& MergeOptions, class AStaticMeshActor*& OutMergedActor);
 };
 
 #if UE_ENABLE_INCLUDE_ORDER_DEPRECATED_IN_5_5
 #include "Engine/MeshMerging.h"
-#endif
\ No newline at end of file
+#endif
diff --git a/Engine/Source/Editor/StaticMeshEditor/Public/StaticMeshEditorSubsystemHelpers.h b/Engine/Source/Editor/StaticMeshEditor/Public/StaticMeshEditorSubsystemHelpers.h
--- a/Engine/Source/Editor/StaticMeshEditor/Public/StaticMeshEditorSubsystemHelpers.h
+++ b/Engine/Source/Editor/StaticMeshEditor/Public/StaticMeshEditorSubsystemHelpers.h
@@ -15,7 +15,7 @@
 class UStaticMeshComponent;
 
 USTRUCT(BlueprintType)
-struct STATICMESHEDITOR_API FStaticMeshReductionSettings
+struct FStaticMeshReductionSettings
 {
 	GENERATED_BODY()
 
@@ -34,7 +34,7 @@ struct STATICMESHEDITOR_API FStaticMeshReductionSettings
 };
 
 USTRUCT(BlueprintType)
-struct STATICMESHEDITOR_API FStaticMeshReductionOptions
+struct FStaticMeshReductionOptions
 {
 	GENERATED_BODY()
 
diff --git a/Engine/Source/Editor/StatusBar/Public/StatusBarSubsystem.h b/Engine/Source/Editor/StatusBar/Public/StatusBarSubsystem.h
--- a/Engine/Source/Editor/StatusBar/Public/StatusBarSubsystem.h
+++ b/Engine/Source/Editor/StatusBar/Public/StatusBarSubsystem.h
@@ -8,6 +8,8 @@
 #include "GlobalStatusBarExtension.h"
 #include "StatusBarSubsystem.generated.h"
 
+#define UE_API STATUSBAR_API
+
 class SStatusBar;
 class SWindow;
 class SWidget;
@@ -66,8 +68,8 @@ struct FStatusBarData
 	TSharedPtr<SWidget> ConsoleEditBox;
 };
 
-UCLASS()
-class STATUSBAR_API UStatusBarSubsystem : public UEditorSubsystem, public IProgressNotificationHandler
+UCLASS(MinimalAPI)
+class UStatusBarSubsystem : public UEditorSubsystem, public IProgressNotificationHandler
 {
 	GENERATED_BODY()
 
@@ -76,12 +78,12 @@ public:
 	/**
 	 *	Prepares for use
 	 */
-	virtual void Initialize(FSubsystemCollectionBase& Collection) override final;
+	UE_API virtual void Initialize(FSubsystemCollectionBase& Collection) override final;
 
 	/**
 	 *	Internal cleanup
 	 */
-	virtual void Deinitialize() override final;
+	UE_API virtual void Deinitialize() override final;
 
 	/**
 	 * Focuses the debug console or opens the output log drawer on the status bar for status bar residing in the passed in parent window 
@@ -90,45 +92,45 @@ public:
 	 * @param bAlwaysToggleDrawer	If true, the output log drawer will be toggled without focusing the debug console first
 	 * @return true if a status bar debug console was found and focused 
 	 */
-	bool ToggleDebugConsole(TSharedRef<SWindow> ParentWindow, bool bAlwaysToggleDrawer=false);
+	UE_API bool ToggleDebugConsole(TSharedRef<SWindow> ParentWindow, bool bAlwaysToggleDrawer=false);
 
 	/**
 	 * Opens the output log drawer for a status bar residing in the active window
 	 *
 	 * @return true if the output log was opened
 	 */
-	bool OpenOutputLogDrawer();
+	UE_API bool OpenOutputLogDrawer();
 
 	/**
 	 * Opens the content browser drawer for a status bar residing in the active window 
 	 * 
 	 * @return true if the content browser was opened
 	 */
-	bool OpenContentBrowserDrawer();
+	UE_API bool OpenContentBrowserDrawer();
 
 	/**
 	 * Opens or closes the content browser drawer for a status bar residing in the active window 
 	 * 
 	 * @return true if the content browser was toggled or false if no status bar in the active window was found
 	 */
-	bool ToggleContentBrowserDrawer();
+	UE_API bool ToggleContentBrowserDrawer();
 
 	/**
 	 * Closes the content browser drawer for a status bar residing in the active window 
 	 * 
 	 * @return true if the content browser was dismissed
 	 */
-	bool DismissContentBrowserDrawer();
+	UE_API bool DismissContentBrowserDrawer();
 
 	/**
 	 * Tries to toggle the given drawer
 	 */
-	bool TryToggleDrawer(const FName DrawerId);
+	UE_API bool TryToggleDrawer(const FName DrawerId);
 
 	/**
 	 * Forces all drawers to dismiss. Usually it dismisses with focus. Only call this if there is some reason an open drawer would be invalid for the current state of the editor.
 	 */
-	bool ForceDismissDrawer();
+	UE_API bool ForceDismissDrawer();
 
 	/** 
 	 * Creates a new instance of a status bar widget
@@ -136,17 +138,17 @@ public:
 	 * @param StatusBarName	The name of the status bar for updating it later. This name must be unique. You can use the instance number of an fname to generate a unique fname for comparison but a non-unique one for serialization of status bar data. I.E all asset editors of a specific type will be saved the same but at runtime will be indentified uniquely. This is usally what you want.
 	 * @param ParentTab	Parent tab of the status bar.
 	 */
-	TSharedRef<SWidget> MakeStatusBarWidget(FName UniqueStatusBarName, const TSharedRef<SDockTab>& InParentTab, UE::Editor::Toolbars::ECreateStatusBarOptions CreateStatusBarOptions = UE::Editor::Toolbars::ECreateStatusBarOptions::Default);
+	UE_API TSharedRef<SWidget> MakeStatusBarWidget(FName UniqueStatusBarName, const TSharedRef<SDockTab>& InParentTab, UE::Editor::Toolbars::ECreateStatusBarOptions CreateStatusBarOptions = UE::Editor::Toolbars::ECreateStatusBarOptions::Default);
 
 	/**
 	 * @return true if a status bar was found for the active window
 	 */
-	bool ActiveWindowHasStatusBar() const;
+	UE_API bool ActiveWindowHasStatusBar() const;
 
 	/**
 	 * @return true if a status bar was found for the active window immediately behind a Notification Window
 	*/
-	bool ActiveWindowBehindNotificationHasStatusBar();
+	UE_API bool ActiveWindowBehindNotificationHasStatusBar();
 
 	/**
 	 * Creates a new instance of a status bar widget
@@ -155,7 +157,7 @@ public:
 	 * @param Drawer		The drawer to add to the status bar
 	 * @param SlotIndex		The position at which to add the new drawer
 	 */
-	void RegisterDrawer(FName StatusBarName, FWidgetDrawerConfig&& Drawer, int32 SlotIndex = INDEX_NONE);
+	UE_API void RegisterDrawer(FName StatusBarName, FWidgetDrawerConfig&& Drawer, int32 SlotIndex = INDEX_NONE);
 
 	/**
 	 * Unregisters and destroys the drawer with the given DrawerId 
@@ -163,7 +165,7 @@ public:
 	 * @param StatusBarName	The name of the status bar to unregister the drawer from
 	 * @param DrawerId		the unique name id of the drawer to unregister
 	 */
-	void UnregisterDrawer(FName StatusBarName, FName DrawerId);
+	UE_API void UnregisterDrawer(FName StatusBarName, FName DrawerId);
 
 	/** 
 	 * Pushes a new status bar message
@@ -173,8 +175,8 @@ public:
 	 * @param InHintText	Optional hint text message.  This message will be highlighted to make it stand out
 	 * @return	A handle to the message for clearing it later
 	 */
-	FStatusBarMessageHandle PushStatusBarMessage(FName StatusBarName, const TAttribute<FText>& InMessage, const TAttribute<FText>& InHintText);
-	FStatusBarMessageHandle PushStatusBarMessage(FName StatusBarName, const TAttribute<FText>& InMessage);
+	UE_API FStatusBarMessageHandle PushStatusBarMessage(FName StatusBarName, const TAttribute<FText>& InMessage, const TAttribute<FText>& InHintText);
+	UE_API FStatusBarMessageHandle PushStatusBarMessage(FName StatusBarName, const TAttribute<FText>& InMessage);
 
 	/**
 	 * Removes a message from the status bar.  When messages are removed the previous message on the stack (if any) is displayed
@@ -182,14 +184,14 @@ public:
 	 * @param StatusBarName	The name of the status bar to remove from
 	 * @param InHandle		Handle to the status bar message to remove
 	 */
-	void PopStatusBarMessage(FName StatusBarName, FStatusBarMessageHandle InHandle);
+	UE_API void PopStatusBarMessage(FName StatusBarName, FStatusBarMessageHandle InHandle);
 
 	/**
 	 * Removes all messages from the status bar
 	 *
 	 * @param StatusBarName	The name of the status bar to remove from
 	 */
-	void ClearStatusBarMessages(FName StatusBarName);
+	UE_API void ClearStatusBarMessages(FName StatusBarName);
 
 	/**
 	 * Registers a new global status bar extension.
@@ -197,7 +199,7 @@ public:
 	 * @param Extension The extension to register
 	 * @return A reference to the extension interface, which can be later used to unregister the extension
 	 */
-	IGlobalStatusBarExtension& RegisterGlobalStatusBarExtension(TUniquePtr<IGlobalStatusBarExtension>&& Extension);
+	UE_API IGlobalStatusBarExtension& RegisterGlobalStatusBarExtension(TUniquePtr<IGlobalStatusBarExtension>&& Extension);
 
 	/**
 	 * Unregisters an existing status bar extension.
@@ -205,14 +207,14 @@ public:
 	 * @param Extension Reference returned by RegisterGlobalStatusBarExtension containing the extension to unregister
 	 * @return An owned instance of the extension that was registered, or null if the extension was not found
 	 */
-	TUniquePtr<IGlobalStatusBarExtension> UnregisterGlobalStatusBarExtension(IGlobalStatusBarExtension* Extension);
+	UE_API TUniquePtr<IGlobalStatusBarExtension> UnregisterGlobalStatusBarExtension(IGlobalStatusBarExtension* Extension);
 
 
 private:
 	/** IProgressNotificationHandler interface */
-	virtual void StartProgressNotification(FProgressNotificationHandle Handle, FText DisplayText, int32 TotalWorkToDo) override;
-	virtual void UpdateProgressNotification(FProgressNotificationHandle Handle, int32 TotalWorkDone, int32 UpdatedTotalWorkToDo, FText UpdatedDisplayText) override;
-	virtual void CancelProgressNotification(FProgressNotificationHandle Handle) override;
+	UE_API virtual void StartProgressNotification(FProgressNotificationHandle Handle, FText DisplayText, int32 TotalWorkToDo) override;
+	UE_API virtual void UpdateProgressNotification(FProgressNotificationHandle Handle, int32 TotalWorkDone, int32 UpdatedTotalWorkToDo, FText UpdatedDisplayText) override;
+	UE_API virtual void CancelProgressNotification(FProgressNotificationHandle Handle) override;
 
 	enum class EDrawerTriggerMode : uint8
 	{
@@ -225,33 +227,35 @@ private:
 	};
 	FRIEND_ENUM_CLASS_FLAGS(EDrawerTriggerMode);
 
-	bool TriggerContentBrowser(EDrawerTriggerMode DrawerTriggerMode);
+	UE_API bool TriggerContentBrowser(EDrawerTriggerMode DrawerTriggerMode);
 
-	void OnDebugConsoleClosed(TWeakPtr<SStatusBar> OwningStatusBar);
-	void CreateContentBrowserIfNeeded();
-	void CreateAndShowNewUserTipIfNeeded(TSharedPtr<SWindow> ParentWindow, bool bIsRunningStartupDialog);
-	const FString GetNewUserTipState() const;
-	void CreateAndShowOneTimeIndustryQueryIfNeeded(TSharedPtr<SWindow> ParentWindow, bool bIsRunningStartupDialog);
-	static const FString GetOneTimeStateWithFallback(const FString StoreId, const FString SectionName, const FString KeyName, const FString FallbackIniLocation, const FString FallbackIniKey);
-	static void SetOneTimeStateWithFallback(const FString StoreId, const FString SectionName, const FString KeyName, const FString FallbackIniLocation, const FString FallbackIniKey);
+	UE_API void OnDebugConsoleClosed(TWeakPtr<SStatusBar> OwningStatusBar);
+	UE_API void CreateContentBrowserIfNeeded();
+	UE_API void CreateAndShowNewUserTipIfNeeded(TSharedPtr<SWindow> ParentWindow, bool bIsRunningStartupDialog);
+	UE_API const FString GetNewUserTipState() const;
+	UE_API void CreateAndShowOneTimeIndustryQueryIfNeeded(TSharedPtr<SWindow> ParentWindow, bool bIsRunningStartupDialog);
+	static UE_API const FString GetOneTimeStateWithFallback(const FString StoreId, const FString SectionName, const FString KeyName, const FString FallbackIniLocation, const FString FallbackIniKey);
+	static UE_API void SetOneTimeStateWithFallback(const FString StoreId, const FString SectionName, const FString KeyName, const FString FallbackIniLocation, const FString FallbackIniKey);
 
-	TSharedPtr<SStatusBar> GetStatusBar(FName StatusBarName) const;
-	TSharedRef<SWidget> OnGetContentBrowser();
-	void OnContentBrowserOpened(FName StatusBarWithDrawerName);
-	void OnContentBrowserDismissed(const TSharedPtr<SWidget>& NewlyFocusedWidget);
-	void HandleDeferredOpenContentBrowser(TSharedPtr<SWindow> ParentWindow);
+	UE_API TSharedPtr<SStatusBar> GetStatusBar(FName StatusBarName) const;
+	UE_API TSharedRef<SWidget> OnGetContentBrowser();
+	UE_API void OnContentBrowserOpened(FName StatusBarWithDrawerName);
+	UE_API void OnContentBrowserDismissed(const TSharedPtr<SWidget>& NewlyFocusedWidget);
+	UE_API void HandleDeferredOpenContentBrowser(TSharedPtr<SWindow> ParentWindow);
 
-	TSharedRef<SWidget> OnGetOutputLog();
-	void OnOutputLogOpened(FName StatusBarWithDrawerName);
-	void OnOutputLogDismised(const TSharedPtr<SWidget>& NewlyFocusedWidget);
+	UE_API TSharedRef<SWidget> OnGetOutputLog();
+	UE_API void OnOutputLogOpened(FName StatusBarWithDrawerName);
+	UE_API void OnOutputLogDismised(const TSharedPtr<SWidget>& NewlyFocusedWidget);
 
-	void OnDebugConsoleDrawerClosed();
+	UE_API void OnDebugConsoleDrawerClosed();
 private:
 	TMap<FName, FStatusBarData> StatusBars;
 	TWeakPtr<SWidget> PreviousKeyboardFocusedWidget;
 	/** The floating content browser that is opened via the content browser button in the status bar */
 	TSharedPtr<SWidget> StatusBarContentBrowser;
 	TSharedPtr<SWidget> StatusBarOutputLog;
-	static int32 MessageHandleCounter;
+	static UE_API int32 MessageHandleCounter;
 	TArray<TUniquePtr<IGlobalStatusBarExtension>> GlobalStatusBarExtensions;
 };
+
+#undef UE_API
diff --git a/Engine/Source/Editor/StorageServerWidgets/Internal/SZenCacheStatistics.h b/Engine/Source/Editor/StorageServerWidgets/Internal/SZenCacheStatistics.h
--- a/Engine/Source/Editor/StorageServerWidgets/Internal/SZenCacheStatistics.h
+++ b/Engine/Source/Editor/StorageServerWidgets/Internal/SZenCacheStatistics.h
@@ -10,7 +10,9 @@
 #include "Widgets/SCompoundWidget.h"
 #include "Widgets/SWidget.h"
 
-class STORAGESERVERWIDGETS_API SZenCacheStatistics : public SCompoundWidget
+#define UE_API STORAGESERVERWIDGETS_API
+
+class SZenCacheStatistics : public SCompoundWidget
 {
 	SLATE_BEGIN_ARGS(SZenCacheStatistics)
 		: _ZenServiceInstance(nullptr)
@@ -20,14 +22,16 @@ class STORAGESERVERWIDGETS_API SZenCacheStatistics : public SCompoundWidget
 
 	SLATE_END_ARGS()
 
-	void Construct(const FArguments& InArgs);
+	UE_API void Construct(const FArguments& InArgs);
 
 private:
 
-	TSharedRef<SWidget> GetGridPanel();
+	UE_API TSharedRef<SWidget> GetGridPanel();
 
-	EActiveTimerReturnType UpdateGridPanels(double InCurrentTime, float InDeltaTime);
+	UE_API EActiveTimerReturnType UpdateGridPanels(double InCurrentTime, float InDeltaTime);
 
 	SVerticalBox::FSlot* GridSlot = nullptr;
 	TAttribute<TSharedPtr<UE::Zen::FZenServiceInstance>> ZenServiceInstance;
-};
\ No newline at end of file
+};
+
+#undef UE_API
diff --git a/Engine/Source/Editor/StorageServerWidgets/Internal/SZenCidStoreStatistics.h b/Engine/Source/Editor/StorageServerWidgets/Internal/SZenCidStoreStatistics.h
--- a/Engine/Source/Editor/StorageServerWidgets/Internal/SZenCidStoreStatistics.h
+++ b/Engine/Source/Editor/StorageServerWidgets/Internal/SZenCidStoreStatistics.h
@@ -10,7 +10,9 @@
 #include "Widgets/SCompoundWidget.h"
 #include "Widgets/SWidget.h"
 
-class STORAGESERVERWIDGETS_API SZenCidStoreStatistics : public SCompoundWidget
+#define UE_API STORAGESERVERWIDGETS_API
+
+class SZenCidStoreStatistics : public SCompoundWidget
 {
 	SLATE_BEGIN_ARGS(SZenCidStoreStatistics)
 		: _ZenServiceInstance(nullptr)
@@ -20,14 +22,16 @@ class STORAGESERVERWIDGETS_API SZenCidStoreStatistics : public SCompoundWidget
 
 	SLATE_END_ARGS()
 
-	void Construct(const FArguments& InArgs);
+	UE_API void Construct(const FArguments& InArgs);
 
 private:
 
-	TSharedRef<SWidget> GetGridPanel();
+	UE_API TSharedRef<SWidget> GetGridPanel();
 
-	EActiveTimerReturnType UpdateGridPanels(double InCurrentTime, float InDeltaTime);
+	UE_API EActiveTimerReturnType UpdateGridPanels(double InCurrentTime, float InDeltaTime);
 
 	SVerticalBox::FSlot* GridSlot = nullptr;
 	TAttribute<TSharedPtr<UE::Zen::FZenServiceInstance>> ZenServiceInstance;
-};
\ No newline at end of file
+};
+
+#undef UE_API
diff --git a/Engine/Source/Editor/StorageServerWidgets/Internal/SZenDialogs.h b/Engine/Source/Editor/StorageServerWidgets/Internal/SZenDialogs.h
--- a/Engine/Source/Editor/StorageServerWidgets/Internal/SZenDialogs.h
+++ b/Engine/Source/Editor/StorageServerWidgets/Internal/SZenDialogs.h
@@ -8,14 +8,18 @@
 #include "Widgets/SBoxPanel.h"
 #include "Widgets/SCompoundWidget.h"
 
+#define UE_API STORAGESERVERWIDGETS_API
+
 class SWidget;
 
-class STORAGESERVERWIDGETS_API SZenStoreStausDialog : public SCompoundWidget
+class SZenStoreStausDialog : public SCompoundWidget
 {
 	SLATE_BEGIN_ARGS(SZenStoreStausDialog) {}
 	SLATE_END_ARGS()
 
-	void Construct(const FArguments& InArgs);
+	UE_API void Construct(const FArguments& InArgs);
 
 private:
 };
+
+#undef UE_API
diff --git a/Engine/Source/Editor/StorageServerWidgets/Internal/SZenProjectStatistics.h b/Engine/Source/Editor/StorageServerWidgets/Internal/SZenProjectStatistics.h
--- a/Engine/Source/Editor/StorageServerWidgets/Internal/SZenProjectStatistics.h
+++ b/Engine/Source/Editor/StorageServerWidgets/Internal/SZenProjectStatistics.h
@@ -10,7 +10,9 @@
 #include "Widgets/SCompoundWidget.h"
 #include "Widgets/SWidget.h"
 
-class STORAGESERVERWIDGETS_API SZenProjectStatistics : public SCompoundWidget
+#define UE_API STORAGESERVERWIDGETS_API
+
+class SZenProjectStatistics : public SCompoundWidget
 {
 	SLATE_BEGIN_ARGS(SZenProjectStatistics)
 		: _ZenServiceInstance(nullptr)
@@ -20,14 +22,16 @@ class STORAGESERVERWIDGETS_API SZenProjectStatistics : public SCompoundWidget
 
 	SLATE_END_ARGS()
 
-	void Construct(const FArguments& InArgs);
+	UE_API void Construct(const FArguments& InArgs);
 
 private:
 
-	TSharedRef<SWidget> GetGridPanel();
+	UE_API TSharedRef<SWidget> GetGridPanel();
 
-	EActiveTimerReturnType UpdateGridPanels(double InCurrentTime, float InDeltaTime);
+	UE_API EActiveTimerReturnType UpdateGridPanels(double InCurrentTime, float InDeltaTime);
 
 	SVerticalBox::FSlot* GridSlot = nullptr;
 	TAttribute<TSharedPtr<UE::Zen::FZenServiceInstance>> ZenServiceInstance;
-};
\ No newline at end of file
+};
+
+#undef UE_API
diff --git a/Engine/Source/Editor/StorageServerWidgets/Internal/SZenServiceStatus.h b/Engine/Source/Editor/StorageServerWidgets/Internal/SZenServiceStatus.h
--- a/Engine/Source/Editor/StorageServerWidgets/Internal/SZenServiceStatus.h
+++ b/Engine/Source/Editor/StorageServerWidgets/Internal/SZenServiceStatus.h
@@ -10,7 +10,9 @@
 #include "Widgets/SCompoundWidget.h"
 #include "Widgets/SWidget.h"
 
-class STORAGESERVERWIDGETS_API SZenServiceStatus : public SCompoundWidget
+#define UE_API STORAGESERVERWIDGETS_API
+
+class SZenServiceStatus : public SCompoundWidget
 {
 	SLATE_BEGIN_ARGS(SZenServiceStatus)
 		: _ZenServiceInstance(nullptr)
@@ -20,7 +22,7 @@ class STORAGESERVERWIDGETS_API SZenServiceStatus : public SCompoundWidget
 
 	SLATE_END_ARGS()
 
-	void Construct(const FArguments& InArgs);
+	UE_API void Construct(const FArguments& InArgs);
 
 private:
 	struct FState
@@ -40,13 +42,15 @@ private:
 	std::atomic<uint32> ActiveStateIndex = 0;
 
 
-	const FState& GetCurrentState() const;
-	TSharedRef<SWidget> GetGridPanel();
-	FReply ExploreDataPath_OnClicked();
-	FReply ExploreExecutablePath_OnClicked();
+	UE_API const FState& GetCurrentState() const;
+	UE_API TSharedRef<SWidget> GetGridPanel();
+	UE_API FReply ExploreDataPath_OnClicked();
+	UE_API FReply ExploreExecutablePath_OnClicked();
 
-	EActiveTimerReturnType UpdateState(double InCurrentTime, float InDeltaTime);
+	UE_API EActiveTimerReturnType UpdateState(double InCurrentTime, float InDeltaTime);
 
 	SVerticalBox::FSlot* GridSlot = nullptr;
 	TAttribute<TSharedPtr<UE::Zen::FZenServiceInstance>> ZenServiceInstance;
-};
\ No newline at end of file
+};
+
+#undef UE_API
diff --git a/Engine/Source/Editor/StorageServerWidgets/Internal/ZenServiceInstanceManager.h b/Engine/Source/Editor/StorageServerWidgets/Internal/ZenServiceInstanceManager.h
--- a/Engine/Source/Editor/StorageServerWidgets/Internal/ZenServiceInstanceManager.h
+++ b/Engine/Source/Editor/StorageServerWidgets/Internal/ZenServiceInstanceManager.h
@@ -4,15 +4,17 @@
 
 #include "Templates/SharedPointer.h"
 
+#define UE_API STORAGESERVERWIDGETS_API
+
 namespace UE::Zen { class FZenServiceInstance; }
 
 namespace UE::Zen
 {
 
-class STORAGESERVERWIDGETS_API FServiceInstanceManager
+class FServiceInstanceManager
 {
 public:
-	TSharedPtr<FZenServiceInstance> GetZenServiceInstance() const;
+	UE_API TSharedPtr<FZenServiceInstance> GetZenServiceInstance() const;
 private:
 	mutable uint16 CurrentPort = 0;
 	mutable TSharedPtr<FZenServiceInstance> CurrentInstance;
@@ -20,3 +22,4 @@ private:
 
 } // namespace UE::Zen
 
+#undef UE_API
diff --git a/Engine/Source/Editor/StringTableEditor/Public/StringTableEditorModule.h b/Engine/Source/Editor/StringTableEditor/Public/StringTableEditorModule.h
--- a/Engine/Source/Editor/StringTableEditor/Public/StringTableEditorModule.h
+++ b/Engine/Source/Editor/StringTableEditor/Public/StringTableEditorModule.h
@@ -7,6 +7,8 @@
 #include "Modules/ModuleInterface.h"
 #include "Toolkits/AssetEditorToolkit.h"
 
+#define UE_API STRINGTABLEEDITOR_API
+
 class IStringTableEditor;
 class UStringTable;
 
@@ -16,13 +18,13 @@ namespace UE::AssetTools
 }
 
 /** String Table Editor module */
-class STRINGTABLEEDITOR_API FStringTableEditorModule : public IModuleInterface,
+class FStringTableEditorModule : public IModuleInterface,
 	public IHasMenuExtensibility, public IHasToolBarExtensibility
 {
 public:
 	// IModuleInterface
-	virtual void StartupModule() override;
-	virtual void ShutdownModule() override;
+	UE_API virtual void StartupModule() override;
+	UE_API virtual void ShutdownModule() override;
 
 	/**
 	 * Creates an instance of string table editor object.
@@ -33,7 +35,7 @@ public:
 	 *
 	 * @return	Interface to the new string table editor
 	 */
-	TSharedRef<IStringTableEditor> CreateStringTableEditor(const EToolkitMode::Type Mode, const TSharedPtr<class IToolkitHost>& InitToolkitHost, UStringTable* StringTable);
+	UE_API TSharedRef<IStringTableEditor> CreateStringTableEditor(const EToolkitMode::Type Mode, const TSharedPtr<class IToolkitHost>& InitToolkitHost, UStringTable* StringTable);
 
 	/** Gets the extensibility managers for outside entities to extend static mesh editor's menus and toolbars */
 	virtual TSharedPtr<FExtensibilityManager> GetMenuExtensibilityManager() override { return MenuExtensibilityManager; }
@@ -41,11 +43,13 @@ public:
 	virtual TSharedPtr<FExtensibilityManager> GetToolBarExtensibilityManager() override { return ToolBarExtensibilityManager; }
 
 	/** StringTable Editor app identifier string */
-	static const FName StringTableEditorAppIdentifier;
+	static UE_API const FName StringTableEditorAppIdentifier;
 
 private:
-	void OnPackageMigration(UE::AssetTools::FPackageMigrationContext& MigrationContext);
+	UE_API void OnPackageMigration(UE::AssetTools::FPackageMigrationContext& MigrationContext);
 
 	TSharedPtr<FExtensibilityManager> MenuExtensibilityManager;
 	TSharedPtr<FExtensibilityManager> ToolBarExtensibilityManager;
 };
+
+#undef UE_API
diff --git a/Engine/Source/Editor/StructUtilsEditor/Private/Customizations/STypeSelector.h b/Engine/Source/Editor/StructUtilsEditor/Private/Customizations/STypeSelector.h
--- a/Engine/Source/Editor/StructUtilsEditor/Private/Customizations/STypeSelector.h
+++ b/Engine/Source/Editor/StructUtilsEditor/Private/Customizations/STypeSelector.h
@@ -4,6 +4,8 @@
 
 #include "SPinTypeSelector.h"
 
+#define UE_API STRUCTUTILSEDITOR_API
+
 /**
  * This widget is a small wrapper around SPinTypeSelector to allow for a right click context menu on the selector pill
  * for the container type so that it can all fit into one compact combo button.
@@ -15,7 +17,7 @@
  * Ex. 'SDoubleComboButton' (SComboButton) which inherits from 'SDoubleMenuAnchor' (SMenuAnchor) or similar with two
  * menu anchors. */
 UE_EXPERIMENTAL(5.6, "This widget is an experimental prototype for StructUtils use only.")
-class STRUCTUTILSEDITOR_API STypeSelector final : public SPinTypeSelector
+class STypeSelector final : public SPinTypeSelector
 {
 public:
 	// Replicated arguments of SPinTypeSelector
@@ -45,15 +47,17 @@ public:
 		SLATE_ARGUMENT(TArray<TSharedPtr<class IPinTypeSelectorFilter>>, CustomFilters)
 	SLATE_END_ARGS()
 
-	void Construct(const FArguments& InArgs, FGetPinTypeTree GetPinTypeTreeFunc);
+	UE_API void Construct(const FArguments& InArgs, FGetPinTypeTree GetPinTypeTreeFunc);
 
 	// SWidget interface
-	virtual FReply OnMouseButtonDown(const FGeometry& MyGeometry, const FPointerEvent& MouseEvent) override;
-	virtual FReply OnMouseButtonUp(const FGeometry& MyGeometry, const FPointerEvent& MouseEvent) override;
+	UE_API virtual FReply OnMouseButtonDown(const FGeometry& MyGeometry, const FPointerEvent& MouseEvent) override;
+	UE_API virtual FReply OnMouseButtonUp(const FGeometry& MyGeometry, const FPointerEvent& MouseEvent) override;
 	// End of SWidget interface
 
-	virtual TSharedRef<SWidget> GetMenuContent(bool bForSecondaryType) override;
+	UE_API virtual TSharedRef<SWidget> GetMenuContent(bool bForSecondaryType) override;
 
 private:
-	FText GetToolTipForSelector() const;
+	UE_API FText GetToolTipForSelector() const;
 };
+
+#undef UE_API
diff --git a/Engine/Source/Editor/StructUtilsEditor/Public/IPropertyBagEditorGraph.h b/Engine/Source/Editor/StructUtilsEditor/Public/IPropertyBagEditorGraph.h
--- a/Engine/Source/Editor/StructUtilsEditor/Public/IPropertyBagEditorGraph.h
+++ b/Engine/Source/Editor/StructUtilsEditor/Public/IPropertyBagEditorGraph.h
@@ -44,7 +44,7 @@ class UPropertyBagEdGraphDragAndDrop : public UInterface
  * Note: The base handler class currently climbs down ownership from Graph->Pin, so if only the graph is valid, it will
  * attempt the graph panel drop through the pins' or nodes' parent graph.
  */
-class STRUCTUTILSEDITOR_API IPropertyBagEdGraphDragAndDrop
+class IPropertyBagEdGraphDragAndDrop
 {
 	GENERATED_BODY()
 
@@ -69,7 +69,7 @@ class UPropertyBagEdGraph : public UPropertyBagEdGraphDragAndDrop
  * Usage: This interface to be inherited by custom subclasses of UEdGraph for full implementation of all StructUtils
  * editor graph integration features.
  */
-class STRUCTUTILSEDITOR_API IPropertyBagEdGraph : public IPropertyBagEdGraphDragAndDrop
+class IPropertyBagEdGraph : public IPropertyBagEdGraphDragAndDrop
 {
 	GENERATED_BODY()
 };
@@ -102,4 +102,4 @@ inline FReply IPropertyBagEdGraphDragAndDrop::OnPropertyBagDetailsDropOnGraphNod
 inline FReply IPropertyBagEdGraphDragAndDrop::OnPropertyBagDetailsDropOnGraph(const FPropertyBagPropertyDesc& PropertyDesc, UEdGraph* Graph, const FVector2f& GraphPosition) const
 {
 	return FReply::Handled();
-}
\ No newline at end of file
+}
diff --git a/Engine/Source/Editor/StructUtilsEditor/Public/IStructUtilsEditor.h b/Engine/Source/Editor/StructUtilsEditor/Public/IStructUtilsEditor.h
--- a/Engine/Source/Editor/StructUtilsEditor/Public/IStructUtilsEditor.h
+++ b/Engine/Source/Editor/StructUtilsEditor/Public/IStructUtilsEditor.h
@@ -6,7 +6,7 @@
 
 
 /** StructUtils Editor public interface */
-class STRUCTUTILSEDITOR_API IStructUtilsEditor : public FAssetEditorToolkit
+class IStructUtilsEditor : public FAssetEditorToolkit
 {
 public:
 
diff --git a/Engine/Source/Editor/StructUtilsEditor/Public/InstancedStructDetails.h b/Engine/Source/Editor/StructUtilsEditor/Public/InstancedStructDetails.h
--- a/Engine/Source/Editor/StructUtilsEditor/Public/InstancedStructDetails.h
+++ b/Engine/Source/Editor/StructUtilsEditor/Public/InstancedStructDetails.h
@@ -5,6 +5,8 @@
 #include "IPropertyTypeCustomization.h"
 #include "IDetailCustomNodeBuilder.h"
 
+#define UE_API STRUCTUTILSEDITOR_API
+
 class IAssetReferenceFilter;
 class IPropertyHandle;
 class IDetailGroup;
@@ -19,27 +21,27 @@ class FInstancedStructProvider;
 /**
  * Type customization for FInstancedStruct.
  */
-class STRUCTUTILSEDITOR_API FInstancedStructDetails : public IPropertyTypeCustomization
+class FInstancedStructDetails : public IPropertyTypeCustomization
 {
 public:
-	virtual ~FInstancedStructDetails() override;
+	UE_API virtual ~FInstancedStructDetails() override;
 	
 	/** Makes a new instance of this detail layout class for a specific detail view requesting it */
-	static TSharedRef<IPropertyTypeCustomization> MakeInstance();
+	static UE_API TSharedRef<IPropertyTypeCustomization> MakeInstance();
 
 	/** IPropertyTypeCustomization interface */
-	virtual void CustomizeHeader(TSharedRef<IPropertyHandle> StructPropertyHandle, class FDetailWidgetRow& HeaderRow, IPropertyTypeCustomizationUtils& StructCustomizationUtils) override;
-	virtual void CustomizeChildren(TSharedRef<IPropertyHandle> StructPropertyHandle, class IDetailChildrenBuilder& StructBuilder, IPropertyTypeCustomizationUtils& StructCustomizationUtils) override;
+	UE_API virtual void CustomizeHeader(TSharedRef<IPropertyHandle> StructPropertyHandle, class FDetailWidgetRow& HeaderRow, IPropertyTypeCustomizationUtils& StructCustomizationUtils) override;
+	UE_API virtual void CustomizeChildren(TSharedRef<IPropertyHandle> StructPropertyHandle, class IDetailChildrenBuilder& StructBuilder, IPropertyTypeCustomizationUtils& StructCustomizationUtils) override;
 
 private:
 
-	void OnCopyInstancedStruct() const;
+	UE_API void OnCopyInstancedStruct() const;
 
-	bool CanPasteInstancedStruct() const;
-	void OnPasteInstancedStruct();
+	UE_API bool CanPasteInstancedStruct() const;
+	UE_API void OnPasteInstancedStruct();
 
 	using FReplacementObjectMap = TMap<UObject*, UObject*>;
-	void OnObjectsReinstanced(const FReplacementObjectMap& ObjectMap);
+	UE_API void OnObjectsReinstanced(const FReplacementObjectMap& ObjectMap);
 
 	/** Handle to the struct property being edited */
 	TSharedPtr<IPropertyHandle> StructProperty;
@@ -57,20 +59,20 @@ private:
  * OnChildRowAdded() is called right after each property is added, which allows the property row to be customizable.
  * Child properties will be grouped if they 1) have "Category" metadata, and 2) have the "EnableCategories" metadata tag.
  */
-class STRUCTUTILSEDITOR_API FInstancedStructDataDetails : public IDetailCustomNodeBuilder, public TSharedFromThis<FInstancedStructDataDetails>
+class FInstancedStructDataDetails : public IDetailCustomNodeBuilder, public TSharedFromThis<FInstancedStructDataDetails>
 {
 public:
-	FInstancedStructDataDetails(TSharedPtr<IPropertyHandle> InStructProperty);
-	virtual ~FInstancedStructDataDetails() override;
+	UE_API FInstancedStructDataDetails(TSharedPtr<IPropertyHandle> InStructProperty);
+	UE_API virtual ~FInstancedStructDataDetails() override;
 
 	//~ Begin IDetailCustomNodeBuilder interface
-	virtual void SetOnRebuildChildren(FSimpleDelegate InOnRegenerateChildren) override;
-	virtual void GenerateHeaderRowContent(FDetailWidgetRow& NodeRow) override;
-	virtual void GenerateChildContent(IDetailChildrenBuilder& ChildrenBuilder) override;
-	virtual void Tick(float DeltaTime) override;
+	UE_API virtual void SetOnRebuildChildren(FSimpleDelegate InOnRegenerateChildren) override;
+	UE_API virtual void GenerateHeaderRowContent(FDetailWidgetRow& NodeRow) override;
+	UE_API virtual void GenerateChildContent(IDetailChildrenBuilder& ChildrenBuilder) override;
+	UE_API virtual void Tick(float DeltaTime) override;
 	virtual bool RequiresTick() const override { return true; }
 	virtual bool InitiallyCollapsed() const override { return false; }
-	virtual FName GetName() const override;
+	UE_API virtual FName GetName() const override;
 	//~ End IDetailCustomNodeBuilder interface
 
 	// Called when a group is added, override to customize a group row.
@@ -79,15 +81,15 @@ public:
 	virtual void OnChildRowAdded(IDetailPropertyRow& ChildRow) {}
 
 private:
-	void OnUserDefinedStructReinstancedHandle(const class UUserDefinedStruct& Struct);
+	UE_API void OnUserDefinedStructReinstancedHandle(const class UUserDefinedStruct& Struct);
 
 	/** Pre/Post change notifications for struct value changes */
-	void OnStructValuePreChange();
-	void OnStructValuePostChange();
-	void OnStructHandlePostChange();
+	UE_API void OnStructValuePreChange();
+	UE_API void OnStructValuePostChange();
+	UE_API void OnStructHandlePostChange();
 
 	/** Returns type of the instanced struct for each instance/object being edited. */
-	TArray<TWeakObjectPtr<const UStruct>> GetInstanceTypes() const;
+	UE_API TArray<TWeakObjectPtr<const UStruct>> GetInstanceTypes() const;
 
 	/**
 	 * Adds groups for the specified properties. One group is created for each unique category (from property metadata) that the properties have.
@@ -95,7 +97,7 @@ private:
 	 * property, and the value is the group. If the property doesn't have a group (category), then it will not have an entry in the map. Note that
 	 * the property must opt-in to grouping by specifying the "EnableCategories" metadata tag.
 	 */
-	void GetPropertyGroups(const TArray<TSharedPtr<IPropertyHandle>>& InProperties, IDetailChildrenBuilder& InChildBuilder, TMap<TSharedPtr<IPropertyHandle>, IDetailGroup*>& OutPropertyToGroup) const;
+	UE_API void GetPropertyGroups(const TArray<TSharedPtr<IPropertyHandle>>& InProperties, IDetailChildrenBuilder& InChildBuilder, TMap<TSharedPtr<IPropertyHandle>, IDetailGroup*>& OutPropertyToGroup) const;
 
 	/** Cached instance types, used to invalidate the layout when types change. */
 	TArray<TWeakObjectPtr<const UStruct>> CachedInstanceTypes;
@@ -112,5 +114,7 @@ private:
 	FDelegateHandle UserDefinedStructReinstancedHandle;
 
 protected:
-	void OnStructLayoutChanges();
+	UE_API void OnStructLayoutChanges();
 };
+
+#undef UE_API
