- MInor cleanup after staticmesh accessor change

diff --git a/Engine/Source/Editor/StaticMeshEditor/Private/StaticMeshEditorTools.cpp b/Engine/Source/Editor/StaticMeshEditor/Private/StaticMeshEditorTools.cpp
--- a/Engine/Source/Editor/StaticMeshEditor/Private/StaticMeshEditorTools.cpp
+++ b/Engine/Source/Editor/StaticMeshEditor/Private/StaticMeshEditorTools.cpp
@@ -5394,7 +5394,7 @@ void FRayTracingProxySettingsLayout::ApplyChanges()
 		SlowTask.MakeDialog();
 
 		StaticMesh->Modify();
-		StaticMesh->GetRayTracingProxySettings() = RayTracingProxySettings;
+		StaticMesh->SetRayTracingProxySettings(RayTracingProxySettings);
 
 		StaticMesh->NotifyRayTracingProxySettingsChanged();
 	}
