StaticMeshEditor: Add support for changing aspect ratio in the static mesh editor viewport

#Synced-CL 48921717

diff --git a/Engine/Source/Editor/StaticMeshEditor/Private/SStaticMeshEditorViewport.cpp b/Engine/Source/Editor/StaticMeshEditor/Private/SStaticMeshEditorViewport.cpp
--- a/Engine/Source/Editor/StaticMeshEditor/Private/SStaticMeshEditorViewport.cpp
+++ b/Engine/Source/Editor/StaticMeshEditor/Private/SStaticMeshEditorViewport.cpp
@@ -37,6 +37,14 @@
 ///////////////////////////////////////////////////////////
 // SStaticMeshEditorViewport
 
+namespace StaticMeshEditorViewportLocals
+{
+	TAutoConsoleVariable<int32> CVarEnableStaticMeshEditorAspectRatioControl(
+		TEXT("StaticMeshEditor.EnableAspectRatioControls"),
+		0,
+		TEXT("Enable a new menu allowing for the fixing of aspect ratio in the Static Mesh Editor's viewport."));
+}
+
 void SStaticMeshEditorViewport::Construct(const FArguments& InArgs)
 {
 	//PreviewScene = new FAdvancedPreviewScene(FPreviewScene::ConstructionValues(), 
@@ -81,6 +89,8 @@ void SStaticMeshEditorViewport::Construct(const FArguments& InArgs)
 	FCoreUObjectDelegates::OnObjectPropertyChanged.AddRaw(this, &SStaticMeshEditorViewport::OnObjectPropertyChanged);
 	
 	UE::AdvancedPreviewScene::BindDefaultOnSettingsChangedHandler(PreviewScene, EditorViewportClient);
+
+	SetAspectRatioState(EStaticMeshEditorAspectRatio::EPlatform);
 }
 
 void SStaticMeshEditorViewport::PopulateViewportOverlays(TSharedRef<SOverlay> Overlay)
@@ -324,6 +334,48 @@ void SStaticMeshEditorViewport::SetFloorAndEnvironmentVisibility(bool bVisible)
 	}
 }
 
+void SStaticMeshEditorViewport::SetAspectRatioState(EStaticMeshEditorAspectRatio AspectRatioIn)
+{
+	if (CurrentAspectRatioMode != AspectRatioIn)
+	{
+		CurrentAspectRatioMode = AspectRatioIn;
+		switch (CurrentAspectRatioMode)
+		{
+		case EStaticMeshEditorAspectRatio::EFree:
+			UnsetAspectRatio();
+			break;
+		case EStaticMeshEditorAspectRatio::EPlatform:
+			SetUsePlatformAspectRatio();
+			break;
+		case EStaticMeshEditorAspectRatio::E6_13:
+			SetAspectRatio(6.0 / 13.0);
+			break;
+		case EStaticMeshEditorAspectRatio::E13_6:
+			SetAspectRatio(13.0 / 6.0);
+			break;
+		case EStaticMeshEditorAspectRatio::E1_1:
+			SetAspectRatio(1.0);
+			break;
+		case EStaticMeshEditorAspectRatio::E4_3:
+			SetAspectRatio(4.0 / 3.0);
+			break;
+		case EStaticMeshEditorAspectRatio::E16_10:
+			SetAspectRatio(16.0 / 10.0);
+			break;
+		case EStaticMeshEditorAspectRatio::E16_9:
+			SetAspectRatio(16.0 / 9.0);
+			break;
+		default:
+			ensure(false);
+		}
+	}
+}
+
+bool SStaticMeshEditorViewport::IsAspectRatio(EStaticMeshEditorAspectRatio TestAspectRatio) const
+{
+	return CurrentAspectRatioMode == TestAspectRatio;
+}
+
 void SStaticMeshEditorViewport::UpdatePreviewSocketMeshes()
 {
 	UStaticMesh* const PreviewStaticMesh = PreviewMeshComponent ? ToRawPtr(PreviewMeshComponent->GetStaticMesh()) : nullptr;
@@ -885,6 +937,55 @@ void SStaticMeshEditorViewport::BindCommands()
 		FIsActionChecked::CreateSP(this, &SStaticMeshEditorViewport::IsLODSelected, 0));
 	// all other LODs will be added dynamically
 
+
+	// Aspect Ratio
+	CommandList->MapAction(
+		Commands.SetAspectRatio_Free,
+		FExecuteAction::CreateSP(this, &SStaticMeshEditorViewport::SetAspectRatioState, EStaticMeshEditorAspectRatio::EFree),
+		FCanExecuteAction(),
+		FIsActionChecked::CreateSP(this, &SStaticMeshEditorViewport::IsAspectRatio, EStaticMeshEditorAspectRatio::EFree));
+
+	CommandList->MapAction(
+		Commands.SetAspectRatio_Platform,
+		FExecuteAction::CreateSP(this, &SStaticMeshEditorViewport::SetAspectRatioState, EStaticMeshEditorAspectRatio::EPlatform),
+		FCanExecuteAction(),
+		FIsActionChecked::CreateSP(this, &SStaticMeshEditorViewport::IsAspectRatio, EStaticMeshEditorAspectRatio::EPlatform));
+
+	CommandList->MapAction(
+		Commands.SetAspectRatio_6_13,
+		FExecuteAction::CreateSP(this, &SStaticMeshEditorViewport::SetAspectRatioState, EStaticMeshEditorAspectRatio::E6_13),
+		FCanExecuteAction(),
+		FIsActionChecked::CreateSP(this, &SStaticMeshEditorViewport::IsAspectRatio, EStaticMeshEditorAspectRatio::E6_13));
+
+	CommandList->MapAction(
+		Commands.SetAspectRatio_13_6,
+		FExecuteAction::CreateSP(this, &SStaticMeshEditorViewport::SetAspectRatioState, EStaticMeshEditorAspectRatio::E13_6),
+		FCanExecuteAction(),
+		FIsActionChecked::CreateSP(this, &SStaticMeshEditorViewport::IsAspectRatio, EStaticMeshEditorAspectRatio::E13_6));
+
+	CommandList->MapAction(
+		Commands.SetAspectRatio_1_1,
+		FExecuteAction::CreateSP(this, &SStaticMeshEditorViewport::SetAspectRatioState, EStaticMeshEditorAspectRatio::E1_1),
+		FCanExecuteAction(),
+		FIsActionChecked::CreateSP(this, &SStaticMeshEditorViewport::IsAspectRatio, EStaticMeshEditorAspectRatio::E1_1));
+
+	CommandList->MapAction(
+		Commands.SetAspectRatio_4_3,
+		FExecuteAction::CreateSP(this, &SStaticMeshEditorViewport::SetAspectRatioState, EStaticMeshEditorAspectRatio::E4_3),
+		FCanExecuteAction(),
+		FIsActionChecked::CreateSP(this, &SStaticMeshEditorViewport::IsAspectRatio, EStaticMeshEditorAspectRatio::E4_3));
+
+	CommandList->MapAction(
+		Commands.SetAspectRatio_16_10,
+		FExecuteAction::CreateSP(this, &SStaticMeshEditorViewport::SetAspectRatioState, EStaticMeshEditorAspectRatio::E16_10),
+		FCanExecuteAction(),
+		FIsActionChecked::CreateSP(this, &SStaticMeshEditorViewport::IsAspectRatio, EStaticMeshEditorAspectRatio::E16_10));
+
+	CommandList->MapAction(
+		Commands.SetAspectRatio_16_9,
+		FExecuteAction::CreateSP(this, &SStaticMeshEditorViewport::SetAspectRatioState, EStaticMeshEditorAspectRatio::E16_9),
+		FCanExecuteAction(),
+		FIsActionChecked::CreateSP(this, &SStaticMeshEditorViewport::IsAspectRatio, EStaticMeshEditorAspectRatio::E16_9));
 }
 
 void SStaticMeshEditorViewport::OnFocusViewportToSelection()
@@ -985,6 +1086,15 @@ TSharedPtr<SWidget> SStaticMeshEditorViewport::BuildViewportToolbar()
 			}
 
 			RightSection.AddEntry(UE::StaticMeshEditor::CreateShowSubmenu());
+
+			RightSection.AddDynamicEntry("StaticMeshEditorAspectRatioControls", FNewToolMenuSectionDelegate::CreateLambda([](FToolMenuSection& InSection)
+			{
+				if (StaticMeshEditorViewportLocals::CVarEnableStaticMeshEditorAspectRatioControl.GetValueOnGameThread())
+				{
+					InSection.AddEntry(UE::StaticMeshEditor::CreateAspectRatioSubmenu());
+				}
+			}));										
+
 			RightSection.AddEntry(UE::StaticMeshEditor::CreateLODSubmenu());
 			RightSection.AddEntry(UE::UnrealEd::CreatePerformanceAndScalabilitySubmenu());
 
diff --git a/Engine/Source/Editor/StaticMeshEditor/Private/SStaticMeshEditorViewport.h b/Engine/Source/Editor/StaticMeshEditor/Private/SStaticMeshEditorViewport.h
--- a/Engine/Source/Editor/StaticMeshEditor/Private/SStaticMeshEditorViewport.h
+++ b/Engine/Source/Editor/StaticMeshEditor/Private/SStaticMeshEditorViewport.h
@@ -22,6 +22,19 @@ class UStaticMesh;
 class UStaticMeshComponent;
 class SRichTextBlock;
 
+enum class EStaticMeshEditorAspectRatio : int32
+{
+	EFree = 0,
+	EPlatform = 1,
+
+	E6_13 = 100,
+	E13_6,
+	E1_1,
+	E4_3,
+	E16_10,
+	E16_9,
+};
+
 /**
  * StaticMesh Editor Preview viewport widget
  */
@@ -204,6 +217,10 @@ private:
 
 	void SetFloorAndEnvironmentVisibility(bool bVisible);
 
+	void SetAspectRatioState(EStaticMeshEditorAspectRatio TestAspectRatio);
+
+	bool IsAspectRatio(EStaticMeshEditorAspectRatio TestAspectRatio) const;
+
 private:
 	
 	/** The parent tab where this viewport resides */
@@ -232,4 +249,7 @@ private:
 
 	/** Handle to the registered OnPreviewFeatureLevelChanged delegate. */
 	FDelegateHandle PreviewFeatureLevelChangedHandle;
+
+	/** Current aspect ratio the viewport is set to.*/
+	EStaticMeshEditorAspectRatio CurrentAspectRatioMode = EStaticMeshEditorAspectRatio::EPlatform;
 };
diff --git a/Engine/Source/Editor/StaticMeshEditor/Private/StaticMeshEditorActions.cpp b/Engine/Source/Editor/StaticMeshEditor/Private/StaticMeshEditorActions.cpp
--- a/Engine/Source/Editor/StaticMeshEditor/Private/StaticMeshEditorActions.cpp
+++ b/Engine/Source/Editor/StaticMeshEditor/Private/StaticMeshEditorActions.cpp
@@ -75,6 +75,18 @@ void FStaticMeshEditorCommands::RegisterCommands()
 
 	UI_COMMAND(SaveGeneratedLODs, "Save Generated LODs", "Saves the generated LODs in the mesh package.", EUserInterfaceActionType::Button, FInputChord());
 	UI_COMMAND(BakeMaterials, "Bake Out Materials", "Bake out Materials for given LOD(s).", EUserInterfaceActionType::Button, FInputChord());
+
+	UI_COMMAND(SetAspectRatio_Free, "Unconstrained", "Sets the aspect ratio to be unconstrained.", EUserInterfaceActionType::RadioButton, FInputChord());
+	UI_COMMAND(SetAspectRatio_Platform, "Platform", "Sets the aspect ratio to be controlled by the current preview platform, if any.", EUserInterfaceActionType::RadioButton, FInputChord());
+	UI_COMMAND(SetAspectRatio_6_13, "6:13", "Sets the aspect ratio to be 6:13 (Mobile Portrait Orientation).", EUserInterfaceActionType::RadioButton, FInputChord());
+	UI_COMMAND(SetAspectRatio_13_6, "13:6", "Sets the aspect ratio to be 6:13 (Mobile Landscape Orientation).", EUserInterfaceActionType::RadioButton, FInputChord());
+	UI_COMMAND(SetAspectRatio_1_1, "1:1", "Sets the aspect ratio to be 1:1 (Square).", EUserInterfaceActionType::RadioButton, FInputChord());
+	UI_COMMAND(SetAspectRatio_4_3, "4:3", "Sets the aspect ratio to be 4:3 (Fullscreen).", EUserInterfaceActionType::RadioButton, FInputChord());
+	UI_COMMAND(SetAspectRatio_16_10, "16:10", "Sets the aspect ratio to be 16:10 (Laptop).", EUserInterfaceActionType::RadioButton, FInputChord());
+	UI_COMMAND(SetAspectRatio_16_9, "16:9", "Sets the aspect ratio to be 16:9 (HD)", EUserInterfaceActionType::RadioButton, FInputChord());	
+	
+
+
 }
 
 #undef LOCTEXT_NAMESPACE
diff --git a/Engine/Source/Editor/StaticMeshEditor/Private/StaticMeshEditorViewportToolbarSections.cpp b/Engine/Source/Editor/StaticMeshEditor/Private/StaticMeshEditorViewportToolbarSections.cpp
--- a/Engine/Source/Editor/StaticMeshEditor/Private/StaticMeshEditorViewportToolbarSections.cpp
+++ b/Engine/Source/Editor/StaticMeshEditor/Private/StaticMeshEditorViewportToolbarSections.cpp
@@ -55,6 +55,25 @@ FToolMenuEntry UE::StaticMeshEditor::CreateShowSubmenu()
 	return UE::UnrealEd::CreateShowSubmenu(FNewToolMenuDelegate::CreateLambda(&FillShowSubmenu));
 }
 
+FToolMenuEntry UE::StaticMeshEditor::CreateAspectRatioSubmenu()
+{
+	const FNewToolMenuChoice InSubmenuChoice = FNewToolMenuDelegate::CreateLambda(&FillAspectRatioSubmenu);
+
+	FToolMenuEntry Entry = FToolMenuEntry::InitSubMenu(
+		"Aspect Ratio",
+		LOCTEXT("AspectSubmenuLabel", "Aspect Ratio"),
+		LOCTEXT("AspectSubmenuTooltip", "Control the aspect ratio of the current viewport"),
+		InSubmenuChoice
+	);
+
+	Entry.Icon = FSlateIcon(FAppStyle::GetAppStyleSetName(), "LevelViewport.ToggleAllowConstrainedAspectRatioInPreview");
+	Entry.ToolBarData.LabelOverride = FText();
+	Entry.ToolBarData.ResizeParams.ClippingPriority = 800;
+
+	return Entry;
+}
+
+
 TSharedRef<SWidget> UE::StaticMeshEditor::GenerateLODMenuWidget(const TSharedPtr<SStaticMeshEditorViewport>& InStaticMeshEditorViewport
 )
 {
@@ -183,4 +202,36 @@ void UE::StaticMeshEditor::FillShowSubmenu(UToolMenu* InMenu)
 	}
 }
 
+void UE::StaticMeshEditor::FillAspectRatioSubmenu(UToolMenu* InMenu)
+{
+	if (UUnrealEdViewportToolbarContext* const EditorViewportContext =
+		InMenu->FindContext<UUnrealEdViewportToolbarContext>())
+	{
+		if (TSharedPtr<SStaticMeshEditorViewport> StaticMeshEditorViewport =
+			StaticCastSharedPtr<SStaticMeshEditorViewport>(EditorViewportContext->Viewport.Pin()))
+		{
+			FToolMenuSection& UnnamedSection = InMenu->FindOrAddSection(NAME_None);
+			const FStaticMeshEditorCommands& Commands = FStaticMeshEditorCommands::Get();
+
+			UnnamedSection.AddMenuEntry(Commands.SetAspectRatio_Free);
+			UnnamedSection.AddMenuEntry(Commands.SetAspectRatio_Platform);
+
+			FToolMenuSection& MobileAspectRatiosSection =
+				InMenu->FindOrAddSection("MobileAspectRatios", LOCTEXT("MobileAspectRatios", "Mobile Aspect Ratios"));
+
+			MobileAspectRatiosSection.AddMenuEntry(Commands.SetAspectRatio_6_13);
+			MobileAspectRatiosSection.AddMenuEntry(Commands.SetAspectRatio_13_6);
+
+			FToolMenuSection& DesktopAspectRatiosSection =
+				InMenu->FindOrAddSection("DesktopAspectRatios", LOCTEXT("DesktopAspectRatios", "Desktop Aspect Ratios"));
+
+			DesktopAspectRatiosSection.AddMenuEntry(Commands.SetAspectRatio_1_1);
+			DesktopAspectRatiosSection.AddMenuEntry(Commands.SetAspectRatio_4_3);
+			DesktopAspectRatiosSection.AddMenuEntry(Commands.SetAspectRatio_16_10);
+			DesktopAspectRatiosSection.AddMenuEntry(Commands.SetAspectRatio_16_9);
+		}
+	}
+}
+
+
 #undef LOCTEXT_NAMESPACE
diff --git a/Engine/Source/Editor/StaticMeshEditor/Private/StaticMeshEditorViewportToolbarSections.h b/Engine/Source/Editor/StaticMeshEditor/Private/StaticMeshEditorViewportToolbarSections.h
--- a/Engine/Source/Editor/StaticMeshEditor/Private/StaticMeshEditorViewportToolbarSections.h
+++ b/Engine/Source/Editor/StaticMeshEditor/Private/StaticMeshEditorViewportToolbarSections.h
@@ -1,4 +1,4 @@
-﻿// Copyright Epic Games, Inc. All Rights Reserved.
+// Copyright Epic Games, Inc. All Rights Reserved.
 
 #pragma once
 
@@ -12,11 +12,13 @@ namespace UE::StaticMeshEditor
 {
 FToolMenuEntry CreateLODSubmenu();
 FToolMenuEntry CreateShowSubmenu();
+FToolMenuEntry CreateAspectRatioSubmenu();
 
 TSharedRef<SWidget> GenerateLODMenuWidget(const TSharedPtr<SStaticMeshEditorViewport>& InStaticMeshEditorViewport);
 TSharedRef<SWidget> GenerateShowMenuWidget(const TSharedPtr<SStaticMeshEditorViewport>& InStaticMeshEditorViewport);
 
 void FillShowSubmenu(UToolMenu *InMenu);
+void FillAspectRatioSubmenu(UToolMenu* InMenu);
 
 FText GetLODMenuLabel(const TSharedPtr<SStaticMeshEditorViewport>& InStaticMeshEditorViewport);
 }
diff --git a/Engine/Source/Editor/StaticMeshEditor/Public/StaticMeshEditorActions.h b/Engine/Source/Editor/StaticMeshEditor/Public/StaticMeshEditorActions.h
--- a/Engine/Source/Editor/StaticMeshEditor/Public/StaticMeshEditorActions.h
+++ b/Engine/Source/Editor/StaticMeshEditor/Public/StaticMeshEditorActions.h
@@ -104,7 +104,18 @@ public:
 
 	TSharedPtr<FUICommandInfo> SaveGeneratedLODs;
 
-	/**
+	// Viewport Aspect Ratio Commands
+
+	TSharedPtr<FUICommandInfo> SetAspectRatio_Free;
+	TSharedPtr<FUICommandInfo> SetAspectRatio_Platform;
+	TSharedPtr<FUICommandInfo> SetAspectRatio_6_13;
+	TSharedPtr<FUICommandInfo> SetAspectRatio_13_6;
+	TSharedPtr<FUICommandInfo> SetAspectRatio_1_1;
+	TSharedPtr<FUICommandInfo> SetAspectRatio_4_3;
+	TSharedPtr<FUICommandInfo> SetAspectRatio_16_10;
+	TSharedPtr<FUICommandInfo> SetAspectRatio_16_9;
+	
+	 /**
 	 * Initialize commands
 	 */
 	virtual void RegisterCommands() override;
