StaticMeshEditor: Add support for a LOD visualizer widget, displaying the current status of the active LOD and/or Nanite.

#Synced-CL 48935287

diff --git a/Engine/Source/Editor/StaticMeshEditor/Private/SStaticMeshEditorViewport.cpp b/Engine/Source/Editor/StaticMeshEditor/Private/SStaticMeshEditorViewport.cpp
--- a/Engine/Source/Editor/StaticMeshEditor/Private/SStaticMeshEditorViewport.cpp
+++ b/Engine/Source/Editor/StaticMeshEditor/Private/SStaticMeshEditorViewport.cpp
@@ -29,6 +29,7 @@
 #include "ViewportToolbar/UnrealEdViewportToolbar.h"
 #include "Widgets/Text/SRichTextBlock.h"
 #include "StaticMeshResources.h"
+#include "SStaticMeshLODVisualizerWidget.h"
 
 #define HITPROXY_SOCKET	1
 
@@ -43,6 +44,11 @@ namespace StaticMeshEditorViewportLocals
 		TEXT("StaticMeshEditor.EnableAspectRatioControls"),
 		0,
 		TEXT("Enable a new menu allowing for the fixing of aspect ratio in the Static Mesh Editor's viewport."));
+
+	TAutoConsoleVariable<int32> CVarEnableStaticMeshEditorLODVisualizer(
+		TEXT("StaticMeshEditor.EnableLODVisualizer"),
+		0,
+		TEXT("Enable a new widget in the static mesh editor viewport to display the status of the current LOD."));
 }
 
 void SStaticMeshEditorViewport::Construct(const FArguments& InArgs)
@@ -113,11 +119,23 @@ void SStaticMeshEditorViewport::PopulateViewportOverlays(TSharedRef<SOverlay> Ov
 	// this widget will display the current viewed feature level
 	Overlay->AddSlot()
 		.VAlign(VAlign_Bottom)
-		.HAlign(HAlign_Right)
+		.HAlign(HAlign_Left)
 		.Padding(5.0f)
 		[
 			BuildShaderPlatformWidget()
 		];
+
+	if (StaticMeshEditorViewportLocals::CVarEnableStaticMeshEditorLODVisualizer.GetValueOnGameThread())
+	{
+		Overlay->AddSlot()
+			.VAlign(EVerticalAlignment::VAlign_Fill)
+			.HAlign(HAlign_Right)
+			.Padding(5.0f)
+			[
+				SAssignNew(LODVisualizer, SStaticMeshLODVisualizerWidget)
+			];
+	}
+	
 }
 
 SStaticMeshEditorViewport::SStaticMeshEditorViewport()
@@ -185,6 +203,35 @@ void SStaticMeshEditorViewport::RefreshViewport()
 {
 	// Invalidate the viewport's display.
 	SceneViewport->Invalidate();
+
+	if (LODVisualizer)
+	{
+		int32 NumLODs = GetLODCount();
+		int32 CurrentLOD = GetCurrentLOD();
+
+#if WITH_EDITORONLY_DATA && WITH_EDITOR
+		FName PlatformName;
+		// Lookup the override preview platform info, if any
+		// @todo this doesn't set PlatformName, just a group, but GetValueForPlatform() will technically work being given a Group name instead of a platform name, so we just use it
+		if (UObject::OnGetPreviewPlatform.IsBound() && UObject::OnGetPreviewPlatform.Execute(PlatformName))
+		{
+			LODVisualizer->SetCurrentPlatform(PlatformName.ToString());
+		}		
+#endif
+
+		LODVisualizer->SetNumLODs(NumLODs);
+
+		LODVisualizer->SetActiveLOD(CurrentLOD);
+		LODVisualizer->SetMinimumLOD(GetCurrentMinimumLOD());
+
+		LODVisualizer->SetNaniteEnabled(IsShowNaniteFallbackVisible());
+		LODVisualizer->SetNaniteActive(!IsShowNaniteFallbackChecked());
+
+		for (int32 LODIndex = 0; LODIndex < NumLODs; ++LODIndex)
+		{
+			LODVisualizer->SetLODPercentage(LODIndex, GetScreenPercentageForLOD(LODIndex));
+		}
+	}
 }
 
 void SStaticMeshEditorViewport::OnObjectPropertyChanged(UObject* ObjectBeingModified, FPropertyChangedEvent& PropertyChangedEvent)
@@ -242,6 +289,13 @@ void SStaticMeshEditorViewport::ToggleShowNaniteFallback()
 			ToggleShowRayTracingFallback();
 		}
 	}
+
+	if (LODVisualizer)
+	{
+		LODVisualizer->SetNaniteEnabled(IsShowNaniteFallbackVisible());
+		LODVisualizer->SetNaniteActive(!IsShowNaniteFallbackChecked());
+	}
+
 }
 
 bool SStaticMeshEditorViewport::IsShowNaniteFallbackChecked() const
@@ -744,6 +798,68 @@ void SStaticMeshEditorViewport::OnLODModelChanged()
 	}
 }
 
+void SStaticMeshEditorViewport::SetCurrentScreenPercentage(float ScreenPercentage)
+{
+	CurrentScreenPercentage = ScreenPercentage;
+
+	if (LODVisualizer)
+	{
+		LODVisualizer->SetScreenPercentage(CurrentScreenPercentage);
+	}
+}
+
+float SStaticMeshEditorViewport::GetScreenPercentageForLOD(int LODIndex)
+{
+	float ScreenSize = 0.0;
+	FName PlatformGroupName = NAME_None;
+
+	if (PreviewMeshComponent)
+	{
+		UStaticMesh* Mesh = PreviewMeshComponent->GetStaticMesh();
+		if (Mesh)
+		{
+			if (Mesh->GetAutoComputeLODScreenSize() && Mesh->GetRenderData())
+			{
+				ScreenSize = Mesh->GetRenderData()->ScreenSize[LODIndex].Default;
+				const float* PlatformScreenSize = Mesh->GetRenderData()->ScreenSize[LODIndex].PerPlatform.Find(PlatformGroupName);
+				if (PlatformScreenSize != nullptr)
+				{
+					ScreenSize = *PlatformScreenSize;
+				}
+
+			}
+			else if (Mesh->IsSourceModelValid(LODIndex))
+			{
+				ScreenSize = Mesh->GetSourceModel(LODIndex).ScreenSize.Default;
+				const float* PlatformScreenSize = Mesh->GetSourceModel(LODIndex).ScreenSize.PerPlatform.Find(PlatformGroupName);
+				if (PlatformScreenSize != nullptr)
+				{
+					ScreenSize = *PlatformScreenSize;
+				}
+			}
+		}
+	}
+	return ScreenSize;
+}
+
+int32 SStaticMeshEditorViewport::GetCurrentMinimumLOD() const
+{
+	if (PreviewMeshComponent)
+	{
+		if (IsShowRayTracingFallbackChecked())
+		{
+			return 0;
+		}
+		else if(PreviewMeshComponent->GetStaticMesh())
+		{
+			return PreviewMeshComponent->GetStaticMesh()->GetMinLODIdx();
+		}
+	}
+
+	return 0;
+}
+
+
 int32 SStaticMeshEditorViewport::GetLODCount() const
 {
 	if (PreviewMeshComponent && PreviewMeshComponent->GetStaticMesh())
diff --git a/Engine/Source/Editor/StaticMeshEditor/Private/SStaticMeshEditorViewport.h b/Engine/Source/Editor/StaticMeshEditor/Private/SStaticMeshEditorViewport.h
--- a/Engine/Source/Editor/StaticMeshEditor/Private/SStaticMeshEditorViewport.h
+++ b/Engine/Source/Editor/StaticMeshEditor/Private/SStaticMeshEditorViewport.h
@@ -21,6 +21,7 @@ class SVerticalBox;
 class UStaticMesh;
 class UStaticMeshComponent;
 class SRichTextBlock;
+class SStaticMeshLODVisualizerWidget;
 
 enum class EStaticMeshEditorAspectRatio : int32
 {
@@ -92,6 +93,10 @@ public:
 
 	void OnLODModelChanged();
 
+	void SetCurrentScreenPercentage(float ScreenPercentage);
+	float GetScreenPercentageForLOD(int LODIndex);
+	int32 GetCurrentMinimumLOD() const;
+
 	/** Retrieves the static mesh component. */
 	UStaticMeshComponent* GetStaticMeshComponent() const;
 
@@ -244,6 +249,12 @@ private:
 	/** Pointer to the box into which the overlay text items are added */
 	TSharedPtr<SRichTextBlock> OverlayText;
 
+	/**  */
+	TSharedPtr<SStaticMeshLODVisualizerWidget> LODVisualizer;
+
+	/** */
+	float CurrentScreenPercentage;
+
 	/** Current LOD Selection where 0 is Auto */
 	int32 LODSelection;
 
diff --git a/Engine/Source/Editor/StaticMeshEditor/Private/SStaticMeshLODVisualizerWidget.cpp b/Engine/Source/Editor/StaticMeshEditor/Private/SStaticMeshLODVisualizerWidget.cpp
new file mode 100644
--- /dev/null
+++ b/Engine/Source/Editor/StaticMeshEditor/Private/SStaticMeshLODVisualizerWidget.cpp
@@ -0,0 +1,399 @@
+// Copyright Epic Games, Inc. All Rights Reserved.
+
+#include "SStaticMeshLODVisualizerWidget.h"
+
+#include "Framework/Application/SlateApplication.h"
+#include "Input/DragAndDrop.h" // FDragDropEvent
+#include "Styling/SlateTypes.h" // FWindowStyle
+#include "Styling/StyleColors.h"
+#include "Widgets/Layout/SBackgroundBlur.h"
+#include "Widgets/Layout/SBorder.h"
+#include "Widgets/Layout/SBox.h"
+#include "Widgets/SBoxPanel.h" // SVerticalBox
+#include "Widgets/SWindow.h"
+#include "Widgets/Text/STextBlock.h"
+#include "Widgets/Input/SSlider.h"
+#include "Widgets/Layout/SExpandableArea.h"
+#include "StaticMeshEditorStyle.h"
+
+
+#define LOCTEXT_NAMESPACE "SStaticMeshLODVisualizerWidget"
+
+namespace LODVisualizerWidgetLocals
+{
+	static FString RootContent = FPaths::EngineContentDir() / TEXT("Slate");
+
+	static FSliderStyle SliderStyle = FSliderStyle()
+		.SetNormalBarImage(FSlateRoundedBoxBrush(FStyleColors::Input, 2.0f, FStyleColors::Input, 1.0f))
+		.SetHoveredBarImage(FSlateRoundedBoxBrush(FStyleColors::Input, 2.0f, FStyleColors::Input, 1.0f))
+		.SetDisabledBarImage(FSlateRoundedBoxBrush(FStyleColors::Input, 2.0f, FStyleColors::Recessed, 1.0f))
+		.SetNormalThumbImage(*FCoreStyle::Get().GetBrush("TreeArrow_Expanded"))
+		.SetHoveredThumbImage(*FCoreStyle::Get().GetBrush("TreeArrow_Expanded"))
+		.SetDisabledThumbImage(*FCoreStyle::Get().GetBrush("TreeArrow_Expanded"))
+		.SetBarThickness(0.0f);
+
+	static FSlateNoResource NoBrush = FSlateNoResource();
+
+	static FString ImagePath = FPaths::EngineContentDir() / TEXT("Slate") / TEXT("Starship/StaticMeshEditor/DisabledOverlay_x16.png");
+
+	static FSlateImageBrush DisabledImage = FSlateImageBrush(ImagePath,	FVector2d(16, 16), FLinearColor::Gray,
+		ESlateBrushTileType::Both, ESlateBrushImageType::FullColor);
+
+	static FColor LODColors[8] = {
+		FColor(141, 211, 199),
+		FColor(255, 255, 179),
+		FColor(190, 186, 218),
+		FColor(251, 128, 114),
+		FColor(128, 177, 211),
+		FColor(253, 180, 98),
+		FColor(179, 222, 105),
+		FColor(252, 205, 229)
+	};
+}
+
+
+PRAGMA_DISABLE_DEPRECATION_WARNINGS
+void SStaticMeshLODVisualizerWidget::Construct(const FArguments& InArgs)
+{
+	SliderValue.BindLambda([this]() { return ScreenPercentage; });
+	PercentAbove.BindLambda([this]() { return 1-ScreenPercentage; });
+	PercentBelow.BindLambda([this]() { return ScreenPercentage; });
+
+	NaniteOverlayVisibility.BindLambda([this]() { return (IsNaniteOverlayActive()) ? EVisibility::Visible : EVisibility::Collapsed; });
+
+	PerLODData.SetNum(8);
+
+	for (int32 LODIndex = 0; LODIndex < 8; ++LODIndex)
+	{
+		PerLODData[LODIndex].ActivationPercentage.BindLambda([this, LODIndex]() { return ComputeLODSectionFill(LODIndex); });
+		PerLODData[LODIndex].Visibility.BindLambda([this, LODIndex]() { return NumLODs > LODIndex ? EVisibility::Visible : EVisibility::Collapsed; });
+		PerLODData[LODIndex].Color.BindLambda([this, LODIndex]()
+			{
+				float ReductionFactor = ((ActiveLOD == LODIndex || ActiveLOD == -1) && !IsNaniteOverlayActive()) ? 1.0 : 0.5;
+				return FSlateColor(ReductionFactor * LODVisualizerWidgetLocals::LODColors[LODIndex]);
+			});
+		PerLODData[LODIndex].MinimalLODVisibility.BindLambda([this, LODIndex]() { return (MinimumLOD > LODIndex) ? EVisibility::Visible : EVisibility::Hidden; });
+		PerLODData[LODIndex].ToolTip.BindLambda([this, LODIndex]() {
+			if (((ActiveLOD == LODIndex || ActiveLOD == -1) && !IsNaniteOverlayActive()))
+			{
+				if (MinimumLOD > LODIndex)
+				{
+					return LOCTEXT("MinimumLODTooltip", "Disabled on this Platform due to Minimum LOD restrictions");
+				}
+				else
+				{
+					return FText::Format(LOCTEXT("NormalLODTooltip", "LOD {0} Visible At {1}"), LODIndex, PerLODData[LODIndex].ScreenPercentage);
+				}
+			}
+			return FText();
+		});
+		
+	}
+
+	FNumberFormattingOptions Options;
+	Options.MaximumFractionalDigits = 3;
+	Options.MinimumFractionalDigits = 3;
+
+
+	ChildSlot
+		[
+				SNew(SExpandableArea)
+				.MinWidth(80)
+				.BodyBorderImage(FAppStyle::Get().GetBrush("FloatingBorder"))
+				.HeaderContent()
+				[
+					SNew(STextBlock)
+					.Text( LOCTEXT( "WidgetLabel", "Show LOD Meter" ))
+				]
+
+				.BodyContent()
+				[
+					SNew(SBox)
+					.HAlign(EHorizontalAlignment::HAlign_Fill)
+					.VAlign(EVerticalAlignment::VAlign_Fill)
+					[
+						SNew(SVerticalBox)
+						 + SVerticalBox::Slot()
+						.AutoHeight()
+						.HAlign(HAlign_Center)
+						.VAlign(VAlign_Center)
+						[						
+							SNew(STextBlock)
+							.Font(FCoreStyle::Get().GetFontStyle("Font.Bold"))
+							.Text_Lambda([this]() {
+								return FText::Format(LOCTEXT("PlatformLabel", "Platform: {0}"), FText::FromString(CurrentPlatform));
+							})
+						]
+
+						 + SVerticalBox::Slot()
+						.FillHeight(1.0)
+						.HAlign(HAlign_Center)
+						[
+
+							SNew(SHorizontalBox)
+							+ SHorizontalBox::Slot()
+							.AutoWidth()
+							.Padding(0, 8)
+							[
+								SNew(SVerticalBox)							
+									+ SVerticalBox::Slot()
+									.FillHeight(PercentAbove)							
+									[
+										SNew(SBorder)
+										.BorderImage(FStaticMeshEditorStyle::Get().GetBrush("StaticMeshEditor.LODMeter.NoBrush"))
+									]
+									+ SVerticalBox::Slot()
+									.AutoHeight()
+									.HAlign(HAlign_Right)
+									[
+										SNew(SBox)
+										.WidthOverride(30)
+										[
+											SNew(STextBlock)
+											.ColorAndOpacity(FSlateColor(FColor::White))
+											.Text_Lambda([this, Options]()
+											{
+												return FText::Format(LOCTEXT("ScreenPercent_Label", "{0}"), FText::AsNumber(ScreenPercentage, &Options));
+											})
+										]
+									]
+									+ SVerticalBox::Slot()
+									.FillHeight(PercentBelow)
+									[
+										SNew(SBorder)
+										.BorderImage(FStaticMeshEditorStyle::Get().GetBrush("StaticMeshEditor.LODMeter.NoBrush"))
+									]
+							]
+
+							 + SHorizontalBox::Slot()
+							.AutoWidth()
+							[
+								SNew(SSlider)
+								.Orientation(EOrientation::Orient_Vertical)
+								.Locked(true)
+								.Value(SliderValue)
+								.Style(FStaticMeshEditorStyle::Get(), "StaticMeshEditor.LODMeter.SliderStyle")
+							]
+
+							+ SHorizontalBox::Slot()
+							.AutoWidth()
+							[
+								SNew(SBox)
+								.WidthOverride(50)
+								[
+									SNew(SOverlay)
+										+ SOverlay::Slot()
+										.Padding(0, 15)
+									[
+									   SAssignNew(LODSlotDisplay, SVerticalBox)
+ 									]
+
+
+										 + SOverlay::Slot()
+										.Padding(0, 15)
+									[
+										SNew(SHorizontalBox)
+										.Visibility(NaniteOverlayVisibility)
+										.ToolTipText(LOCTEXT("NaniteLODTooltip", "Nanite replaces traditional LODs on this Platform."))
+
+										+ SHorizontalBox::Slot()
+										.FillWidth(1)
+										[
+											SNew(SBorder)
+											.Visibility(EVisibility::Hidden)
+											.BorderImage(FCoreStyle::Get().GetBrush("WhiteBrush"))
+										]
+
+										+ SHorizontalBox::Slot()										
+										.FillWidth(5)
+										[
+											SNew(SBox)
+											.WidthOverride(150)
+											[
+												SNew(SBorder)										
+												.BorderImage(FCoreStyle::Get().GetBrush("WhiteBrush"))
+												.BorderBackgroundColor(FColor(200, 200, 200))
+												.ForegroundColor(FSlateColor(FColor::Black))
+												.VAlign(VAlign_Center)
+												.HAlign(HAlign_Center)
+												[
+													SNew(STextBlock)
+													.RenderTransform(FSlateRenderTransform(FQuat2f(FMath::DegreesToRadians(90)), FVector2f::Zero()))
+													.RenderTransformPivot(FVector2d(.5,.5))																
+													.Font(FCoreStyle::Get().GetFontStyle("Font.Large.Bold"))
+													.Text(LOCTEXT("NaniteLabel", "Nanite LOD"))
+													.Clipping(EWidgetClipping::Inherit)
+
+												]
+											]
+										]
+										+ SHorizontalBox::Slot()
+										.FillWidth(1)
+										[
+											SNew(SBorder)
+											.Visibility(EVisibility::Hidden)
+											.BorderImage(FCoreStyle::Get().GetBrush("WhiteBrush"))
+										]
+
+ 									]
+
+									+ SOverlay::Slot()
+										.Padding(0, 15)
+										[
+											SNew(SVerticalBox)
+												.Visibility(EVisibility::HitTestInvisible)
+												+ SVerticalBox::Slot()
+												.FillHeight(PercentAbove)
+												[
+													SNew(SBorder)
+												]
+												+ SVerticalBox::Slot()
+												.AutoHeight()
+												[
+													SNew(SBorder)
+													.Padding(0.0, 0.5, 0.0, 0.5)
+													.BorderImage(FCoreStyle::Get().GetBrush("WhiteBrush"))
+													.BorderBackgroundColor(FSlateColor(FLinearColor(0,0,0)))
+												]
+												+ SVerticalBox::Slot()
+												.FillHeight(PercentBelow)
+												[
+													SNew(SBorder)
+												]
+										]
+								]
+							]
+
+						]
+				
+					]
+
+				]
+			]
+	;
+
+	static FString ImagePath = FPaths::EngineContentDir() / TEXT("Slate") / TEXT("Starship/StaticMeshEditor/DisabledOverlay_x16.png");
+
+	static FSlateImageBrush DisabledImage = FSlateImageBrush(ImagePath, FVector2d(16, 16));
+
+	for (int32 LODIndex = 0; LODIndex < 8; ++LODIndex)
+	{
+		LODSlotDisplay->AddSlot()
+		.FillHeight(PerLODData[LODIndex].ActivationPercentage)
+		[
+			SNew(SBorder)
+			.Visibility(PerLODData[LODIndex].Visibility)
+			.BorderImage(FCoreStyle::Get().GetBrush("WhiteBrush"))
+			.BorderBackgroundColor(PerLODData[LODIndex].Color)
+			.ForegroundColor(FSlateColor(FColor::Black))
+			.ToolTipText(PerLODData[LODIndex].ToolTip)
+			.VAlign(VAlign_Fill)
+			.HAlign(HAlign_Fill)
+			.Padding(0)
+			[
+				SNew(SOverlay)
+				.ToolTipText(PerLODData[LODIndex].ToolTip)
+				+ SOverlay::Slot()
+				.VAlign(VAlign_Fill)
+				.HAlign(HAlign_Fill)				
+				[
+				SNew(SBorder)
+				.Visibility(PerLODData[LODIndex].MinimalLODVisibility)
+				.BorderImage(FStaticMeshEditorStyle::Get().GetBrush("StaticMeshEditor.LODMeter.DisabledBrush"))
+				.BorderBackgroundColor(PerLODData[LODIndex].Color)
+				.ToolTipText(PerLODData[LODIndex].ToolTip)
+				.VAlign(VAlign_Center)
+				.HAlign(HAlign_Center)
+				]
+
+				+ SOverlay::Slot()
+				.VAlign(VAlign_Center)
+			    .HAlign(HAlign_Center)
+				[
+					SNew(STextBlock)
+					.Text( FText::Format( LOCTEXT("LOD_Label", "LOD {0}"), LODIndex ))
+					.ToolTipText(PerLODData[LODIndex].ToolTip)
+				]
+
+				+ SOverlay::Slot()
+				.VAlign(VAlign_Fill)
+				.HAlign(HAlign_Fill)
+				[
+					SNew(SBox)
+					.ToolTipText(PerLODData[LODIndex].ToolTip)
+				]
+			]
+		];
+	}
+}
+PRAGMA_ENABLE_DEPRECATION_WARNINGS
+
+void SStaticMeshLODVisualizerWidget::SetScreenPercentage(float NewScreenPercentageIn)
+{
+	ScreenPercentage = FMath::Clamp(NewScreenPercentageIn, 0.0, 1.0);
+}
+
+void SStaticMeshLODVisualizerWidget::SetNumLODs(int NumLodsIn)
+{
+	NumLODs = FMath::Clamp(NumLodsIn, 1, 8);
+}
+
+void SStaticMeshLODVisualizerWidget::SetLODPercentage(int32 LODIndex, float NewLODPercentageIn)
+{
+	if (LODIndex < 0 || LODIndex > 7)
+	{
+		return;
+	}
+
+	PerLODData[LODIndex].ScreenPercentage = NewLODPercentageIn;
+}
+
+void SStaticMeshLODVisualizerWidget::SetActiveLOD(int LODIndexIn)
+{
+	ActiveLOD = LODIndexIn;
+}
+
+float SStaticMeshLODVisualizerWidget::ComputeLODSectionFill(int32 LODIndex)
+{
+	float UpperValue = 1;
+	float LowerValue = 0;
+
+	if (LODIndex > 0)
+	{
+		UpperValue = PerLODData[LODIndex].ScreenPercentage;
+	}
+
+	if (LODIndex < 7)
+	{
+		LowerValue = PerLODData[LODIndex + 1].ScreenPercentage;
+	}
+
+	return UpperValue - LowerValue;
+}
+
+void SStaticMeshLODVisualizerWidget::SetNaniteEnabled(bool bNaniteEnabledIn)
+{
+	bNaniteEnabled = bNaniteEnabledIn;
+}
+
+void SStaticMeshLODVisualizerWidget::SetNaniteActive(bool bNaniteActiveIn)
+{
+	bNaniteActive = bNaniteActiveIn;
+}
+
+bool SStaticMeshLODVisualizerWidget::IsNaniteOverlayActive() const
+{
+	return bNaniteEnabled && bNaniteActive;
+}
+
+void SStaticMeshLODVisualizerWidget::SetMinimumLOD(int LODIndexIn)
+{
+	MinimumLOD = FMath::Clamp(LODIndexIn, 0, 7);;
+}
+
+void SStaticMeshLODVisualizerWidget::SetCurrentPlatform(FString PlatformName)
+{
+	CurrentPlatform = PlatformName;
+}
+
+#undef LOCTEXT_NAMESPACE
\ No newline at end of file
diff --git a/Engine/Source/Editor/StaticMeshEditor/Private/SStaticMeshLODVisualizerWidget.h b/Engine/Source/Editor/StaticMeshEditor/Private/SStaticMeshLODVisualizerWidget.h
new file mode 100644
--- /dev/null
+++ b/Engine/Source/Editor/StaticMeshEditor/Private/SStaticMeshLODVisualizerWidget.h
@@ -0,0 +1,73 @@
+// Copyright Epic Games, Inc. All Rights Reserved.
+
+#pragma once
+
+#include "CoreMinimal.h"
+
+#include "Widgets/DeclarativeSyntaxSupport.h"
+#include "Widgets/SCompoundWidget.h"
+
+struct FLODVisualizationData
+{
+	float ScreenPercentage = 0.0;
+	TAttribute<float> ActivationPercentage;
+	TAttribute<EVisibility> Visibility;
+	TAttribute<EVisibility> MinimalLODVisibility;
+	TAttribute<FSlateColor> Color;
+	TAttribute<const FSlateBrush*> ActiveBorder;
+	TAttribute<FText> ToolTip;
+};
+
+class SVerticalBox;
+
+class SStaticMeshLODVisualizerWidget : public SCompoundWidget
+{
+public:
+
+	SLATE_BEGIN_ARGS(SStaticMeshLODVisualizerWidget) {}
+		SLATE_DEFAULT_SLOT(FArguments, Content)
+	SLATE_END_ARGS()
+
+	void Construct(const FArguments & InArgs);
+
+	void SetScreenPercentage(float NewScreenPercentageIn);
+
+	void SetLODPercentage(int32 LODIndex, float NewLODPercentageIn);
+
+	void SetNumLODs(int NumLods);
+
+	void SetActiveLOD(int LODIndexIn);
+
+	void SetNaniteEnabled(bool bNaniteEnabledIn);
+
+	void SetNaniteActive(bool bNaniteActiveIn);
+
+	void SetMinimumLOD(int LODIndexIn);
+
+	void SetCurrentPlatform(FString PlatformName);
+
+protected:
+
+	float ScreenPercentage = 1.f;
+	int32 NumLODs = 1;
+	int32 ActiveLOD = 0;
+	int32 MinimumLOD = 0;
+	FString CurrentPlatform = "NONE";
+
+	bool bNaniteEnabled = false;
+	bool bNaniteActive = false;
+
+	TSharedPtr<SVerticalBox> LODSlotDisplay;
+
+	TAttribute<float> SliderValue;
+	TAttribute<float> PercentAbove;
+	TAttribute<float> PercentBelow;
+
+	TAttribute<EVisibility> NaniteOverlayVisibility;
+
+	TArray<FLODVisualizationData> PerLODData;
+
+	float ComputeLODSectionFill(int32 LODIndex);
+
+	bool IsNaniteOverlayActive() const;
+};
\ No newline at end of file
diff --git a/Engine/Source/Editor/StaticMeshEditor/Private/StaticMeshEditorStyle.cpp b/Engine/Source/Editor/StaticMeshEditor/Private/StaticMeshEditorStyle.cpp
new file mode 100644
--- /dev/null
+++ b/Engine/Source/Editor/StaticMeshEditor/Private/StaticMeshEditorStyle.cpp
@@ -0,0 +1,51 @@
+// Copyright Epic Games, Inc. All Rights Reserved.
+
+#include "StaticMeshEditorStyle.h"
+
+#include "Styling/SlateTypes.h"
+#include "Styling/CoreStyle.h"
+#include "Styling/AppStyle.h"
+#include "Styling/StyleColors.h"
+#include "Styling/SlateStyleMacros.h"
+#include "Styling/ToolBarStyle.h"
+
+FName FStaticMeshEditorStyle::StyleName("StaticMeshEditorStyle");
+
+FStaticMeshEditorStyle::FStaticMeshEditorStyle()
+	: FSlateStyleSet(StyleName)
+{
+	const FVector2D IconSize(16.0f, 16.0f);
+	const FVector2D ToolbarIconSize(20.0f, 20.0f);
+
+	SetContentRoot(FPaths::EngineContentDir() / TEXT("Slate/Starship/StaticMeshEditor"));
+	SetCoreContentRoot(FPaths::EngineContentDir() / TEXT("Slate"));
+
+	{
+		FSliderStyle SliderStyle = FSliderStyle()
+			.SetNormalBarImage(FSlateRoundedBoxBrush(FStyleColors::Input, 2.0f, FStyleColors::Input, 1.0f))
+			.SetHoveredBarImage(FSlateRoundedBoxBrush(FStyleColors::Input, 2.0f, FStyleColors::Input, 1.0f))
+			.SetDisabledBarImage(FSlateRoundedBoxBrush(FStyleColors::Input, 2.0f, FStyleColors::Recessed, 1.0f))
+			.SetNormalThumbImage(*FCoreStyle::Get().GetBrush("TreeArrow_Expanded"))
+			.SetHoveredThumbImage(*FCoreStyle::Get().GetBrush("TreeArrow_Expanded"))
+			.SetDisabledThumbImage(*FCoreStyle::Get().GetBrush("TreeArrow_Expanded"))
+			.SetBarThickness(0.0f);
+		Set("StaticMeshEditor.LODMeter.SliderStyle", SliderStyle);
+	}
+
+	Set("StaticMeshEditor.LODMeter.DisabledBrush", new IMAGE_BRUSH("DisabledOverlay_x16", IconSize, FLinearColor::Gray, ESlateBrushTileType::Both));
+	Set("StaticMeshEditor.LODMeter.NoBrush",new FSlateNoResource());
+
+
+	FSlateStyleRegistry::RegisterSlateStyle(*this);
+}
+
+FStaticMeshEditorStyle::~FStaticMeshEditorStyle()
+{
+	FSlateStyleRegistry::UnRegisterSlateStyle(*this);
+}
+
+FStaticMeshEditorStyle& FStaticMeshEditorStyle::Get()
+{
+	static FStaticMeshEditorStyle Inst;
+	return Inst;
+}
\ No newline at end of file
diff --git a/Engine/Source/Editor/StaticMeshEditor/Private/StaticMeshEditorStyle.h b/Engine/Source/Editor/StaticMeshEditor/Private/StaticMeshEditorStyle.h
new file mode 100644
--- /dev/null
+++ b/Engine/Source/Editor/StaticMeshEditor/Private/StaticMeshEditorStyle.h
@@ -0,0 +1,25 @@
+// Copyright Epic Games, Inc. All Rights Reserved.
+
+#pragma once
+
+#include "Styling/SlateStyle.h"
+#include "Styling/SlateStyleRegistry.h"
+
+/**
+ * Slate style set for Static Mesh Editor
+ */
+class FStaticMeshEditorStyle
+	: public FSlateStyleSet
+{
+public:
+	static FName StyleName;
+
+	/** Access the singleton instance for this style set */
+	static FStaticMeshEditorStyle& Get();
+
+private:
+
+	FStaticMeshEditorStyle();
+	~FStaticMeshEditorStyle();
+};
+
diff --git a/Engine/Source/Editor/StaticMeshEditor/Private/StaticMeshEditorViewportClient.cpp b/Engine/Source/Editor/StaticMeshEditor/Private/StaticMeshEditorViewportClient.cpp
--- a/Engine/Source/Editor/StaticMeshEditor/Private/StaticMeshEditorViewportClient.cpp
+++ b/Engine/Source/Editor/StaticMeshEditor/Private/StaticMeshEditorViewportClient.cpp
@@ -876,6 +876,8 @@ void FStaticMeshEditorViewportClient::DrawCanvas(FViewport& InViewport, FSceneVi
 	FormatOptions.MinimumFractionalDigits = 3;
 	FormatOptions.MaximumFractionalDigits = 6;
 	FormatOptions.MaximumIntegralDigits = 6;
+	StaticMeshEditorViewport->SetCurrentScreenPercentage(CurrentScreenSize);
+
 	TextItems.Emplace(FText::Format(NSLOCTEXT("UnrealEd", "ScreenSize_F", "Current Screen Size:  {0}"), FText::AsNumber(CurrentScreenSize, &FormatOptions)));
 
 	const FText StaticMeshTriangleCount = FText::AsNumber(StaticMeshEditorPtr.Pin()->GetNumTriangles(CurrentLODLevel));
