[Backout] - CL48959309
#rnx
Original CL Desc
-----------------------------------------------------------------
StaticMeshEditor: Add support for a LOD visualizer widget, displaying the current status of the active LOD and/or Nanite.

#Synced-CL 48935287

#Synced-CL 48935287

diff --git a/Engine/Source/Editor/StaticMeshEditor/Private/SStaticMeshEditorViewport.cpp b/Engine/Source/Editor/StaticMeshEditor/Private/SStaticMeshEditorViewport.cpp
--- a/Engine/Source/Editor/StaticMeshEditor/Private/SStaticMeshEditorViewport.cpp
+++ b/Engine/Source/Editor/StaticMeshEditor/Private/SStaticMeshEditorViewport.cpp
@@ -29,7 +29,6 @@
 #include "ViewportToolbar/UnrealEdViewportToolbar.h"
 #include "Widgets/Text/SRichTextBlock.h"
 #include "StaticMeshResources.h"
-#include "SStaticMeshLODVisualizerWidget.h"
 
 #define HITPROXY_SOCKET	1
 
@@ -44,11 +43,6 @@ namespace StaticMeshEditorViewportLocals
 		TEXT("StaticMeshEditor.EnableAspectRatioControls"),
 		0,
 		TEXT("Enable a new menu allowing for the fixing of aspect ratio in the Static Mesh Editor's viewport."));
-
-	TAutoConsoleVariable<int32> CVarEnableStaticMeshEditorLODVisualizer(
-		TEXT("StaticMeshEditor.EnableLODVisualizer"),
-		0,
-		TEXT("Enable a new widget in the static mesh editor viewport to display the status of the current LOD."));
 }
 
 void SStaticMeshEditorViewport::Construct(const FArguments& InArgs)
@@ -119,23 +113,11 @@ void SStaticMeshEditorViewport::PopulateViewportOverlays(TSharedRef<SOverlay> Ov
 	// this widget will display the current viewed feature level
 	Overlay->AddSlot()
 		.VAlign(VAlign_Bottom)
-		.HAlign(HAlign_Left)
+		.HAlign(HAlign_Right)
 		.Padding(5.0f)
 		[
 			BuildShaderPlatformWidget()
 		];
-
-	if (StaticMeshEditorViewportLocals::CVarEnableStaticMeshEditorLODVisualizer.GetValueOnGameThread())
-	{
-		Overlay->AddSlot()
-			.VAlign(EVerticalAlignment::VAlign_Fill)
-			.HAlign(HAlign_Right)
-			.Padding(5.0f)
-			[
-				SAssignNew(LODVisualizer, SStaticMeshLODVisualizerWidget)
-			];
-	}
-	
 }
 
 SStaticMeshEditorViewport::SStaticMeshEditorViewport()
@@ -203,35 +185,6 @@ void SStaticMeshEditorViewport::RefreshViewport()
 {
 	// Invalidate the viewport's display.
 	SceneViewport->Invalidate();
-
-	if (LODVisualizer)
-	{
-		int32 NumLODs = GetLODCount();
-		int32 CurrentLOD = GetCurrentLOD();
-
-#if WITH_EDITORONLY_DATA && WITH_EDITOR
-		FName PlatformName;
-		// Lookup the override preview platform info, if any
-		// @todo this doesn't set PlatformName, just a group, but GetValueForPlatform() will technically work being given a Group name instead of a platform name, so we just use it
-		if (UObject::OnGetPreviewPlatform.IsBound() && UObject::OnGetPreviewPlatform.Execute(PlatformName))
-		{
-			LODVisualizer->SetCurrentPlatform(PlatformName.ToString());
-		}		
-#endif
-
-		LODVisualizer->SetNumLODs(NumLODs);
-
-		LODVisualizer->SetActiveLOD(CurrentLOD);
-		LODVisualizer->SetMinimumLOD(GetCurrentMinimumLOD());
-
-		LODVisualizer->SetNaniteEnabled(IsShowNaniteFallbackVisible());
-		LODVisualizer->SetNaniteActive(!IsShowNaniteFallbackChecked());
-
-		for (int32 LODIndex = 0; LODIndex < NumLODs; ++LODIndex)
-		{
-			LODVisualizer->SetLODPercentage(LODIndex, GetScreenPercentageForLOD(LODIndex));
-		}
-	}
 }
 
 void SStaticMeshEditorViewport::OnObjectPropertyChanged(UObject* ObjectBeingModified, FPropertyChangedEvent& PropertyChangedEvent)
@@ -289,13 +242,6 @@ void SStaticMeshEditorViewport::ToggleShowNaniteFallback()
 			ToggleShowRayTracingFallback();
 		}
 	}
-
-	if (LODVisualizer)
-	{
-		LODVisualizer->SetNaniteEnabled(IsShowNaniteFallbackVisible());
-		LODVisualizer->SetNaniteActive(!IsShowNaniteFallbackChecked());
-	}
-
 }
 
 bool SStaticMeshEditorViewport::IsShowNaniteFallbackChecked() const
@@ -798,68 +744,6 @@ void SStaticMeshEditorViewport::OnLODModelChanged()
 	}
 }
 
-void SStaticMeshEditorViewport::SetCurrentScreenPercentage(float ScreenPercentage)
-{
-	CurrentScreenPercentage = ScreenPercentage;
-
-	if (LODVisualizer)
-	{
-		LODVisualizer->SetScreenPercentage(CurrentScreenPercentage);
-	}
-}
-
-float SStaticMeshEditorViewport::GetScreenPercentageForLOD(int LODIndex)
-{
-	float ScreenSize = 0.0;
-	FName PlatformGroupName = NAME_None;
-
-	if (PreviewMeshComponent)
-	{
-		UStaticMesh* Mesh = PreviewMeshComponent->GetStaticMesh();
-		if (Mesh)
-		{
-			if (Mesh->GetAutoComputeLODScreenSize() && Mesh->GetRenderData())
-			{
-				ScreenSize = Mesh->GetRenderData()->ScreenSize[LODIndex].Default;
-				const float* PlatformScreenSize = Mesh->GetRenderData()->ScreenSize[LODIndex].PerPlatform.Find(PlatformGroupName);
-				if (PlatformScreenSize != nullptr)
-				{
-					ScreenSize = *PlatformScreenSize;
-				}
-
-			}
-			else if (Mesh->IsSourceModelValid(LODIndex))
-			{
-				ScreenSize = Mesh->GetSourceModel(LODIndex).ScreenSize.Default;
-				const float* PlatformScreenSize = Mesh->GetSourceModel(LODIndex).ScreenSize.PerPlatform.Find(PlatformGroupName);
-				if (PlatformScreenSize != nullptr)
-				{
-					ScreenSize = *PlatformScreenSize;
-				}
-			}
-		}
-	}
-	return ScreenSize;
-}
-
-int32 SStaticMeshEditorViewport::GetCurrentMinimumLOD() const
-{
-	if (PreviewMeshComponent)
-	{
-		if (IsShowRayTracingFallbackChecked())
-		{
-			return 0;
-		}
-		else if(PreviewMeshComponent->GetStaticMesh())
-		{
-			return PreviewMeshComponent->GetStaticMesh()->GetMinLODIdx();
-		}
-	}
-
-	return 0;
-}
-
-
 int32 SStaticMeshEditorViewport::GetLODCount() const
 {
 	if (PreviewMeshComponent && PreviewMeshComponent->GetStaticMesh())
diff --git a/Engine/Source/Editor/StaticMeshEditor/Private/SStaticMeshEditorViewport.h b/Engine/Source/Editor/StaticMeshEditor/Private/SStaticMeshEditorViewport.h
--- a/Engine/Source/Editor/StaticMeshEditor/Private/SStaticMeshEditorViewport.h
+++ b/Engine/Source/Editor/StaticMeshEditor/Private/SStaticMeshEditorViewport.h
@@ -21,7 +21,6 @@ class SVerticalBox;
 class UStaticMesh;
 class UStaticMeshComponent;
 class SRichTextBlock;
-class SStaticMeshLODVisualizerWidget;
 
 enum class EStaticMeshEditorAspectRatio : int32
 {
@@ -93,10 +92,6 @@ public:
 
 	void OnLODModelChanged();
 
-	void SetCurrentScreenPercentage(float ScreenPercentage);
-	float GetScreenPercentageForLOD(int LODIndex);
-	int32 GetCurrentMinimumLOD() const;
-
 	/** Retrieves the static mesh component. */
 	UStaticMeshComponent* GetStaticMeshComponent() const;
 
@@ -249,12 +244,6 @@ private:
 	/** Pointer to the box into which the overlay text items are added */
 	TSharedPtr<SRichTextBlock> OverlayText;
 
-	/**  */
-	TSharedPtr<SStaticMeshLODVisualizerWidget> LODVisualizer;
-
-	/** */
-	float CurrentScreenPercentage;
-
 	/** Current LOD Selection where 0 is Auto */
 	int32 LODSelection;
 
diff --git a/Engine/Source/Editor/StaticMeshEditor/Private/SStaticMeshLODVisualizerWidget.cpp b/Engine/Source/Editor/StaticMeshEditor/Private/SStaticMeshLODVisualizerWidget.cpp
deleted file mode 100644
diff --git a/Engine/Source/Editor/StaticMeshEditor/Private/SStaticMeshLODVisualizerWidget.h b/Engine/Source/Editor/StaticMeshEditor/Private/SStaticMeshLODVisualizerWidget.h
deleted file mode 100644
diff --git a/Engine/Source/Editor/StaticMeshEditor/Private/StaticMeshEditorStyle.cpp b/Engine/Source/Editor/StaticMeshEditor/Private/StaticMeshEditorStyle.cpp
deleted file mode 100644
diff --git a/Engine/Source/Editor/StaticMeshEditor/Private/StaticMeshEditorStyle.h b/Engine/Source/Editor/StaticMeshEditor/Private/StaticMeshEditorStyle.h
deleted file mode 100644
diff --git a/Engine/Source/Editor/StaticMeshEditor/Private/StaticMeshEditorViewportClient.cpp b/Engine/Source/Editor/StaticMeshEditor/Private/StaticMeshEditorViewportClient.cpp
--- a/Engine/Source/Editor/StaticMeshEditor/Private/StaticMeshEditorViewportClient.cpp
+++ b/Engine/Source/Editor/StaticMeshEditor/Private/StaticMeshEditorViewportClient.cpp
@@ -876,8 +876,6 @@ void FStaticMeshEditorViewportClient::DrawCanvas(FViewport& InViewport, FSceneVi
 	FormatOptions.MinimumFractionalDigits = 3;
 	FormatOptions.MaximumFractionalDigits = 6;
 	FormatOptions.MaximumIntegralDigits = 6;
-	StaticMeshEditorViewport->SetCurrentScreenPercentage(CurrentScreenSize);
-
 	TextItems.Emplace(FText::Format(NSLOCTEXT("UnrealEd", "ScreenSize_F", "Current Screen Size:  {0}"), FText::AsNumber(CurrentScreenSize, &FormatOptions)));
 
 	const FText StaticMeshTriangleCount = FText::AsNumber(StaticMeshEditorPtr.Pin()->GetNumTriangles(CurrentLODLevel));
