EditorUsability : UnrealEd
- Added keybind to toggle the grid in Asset Editor (ALT+SHIFT+G since ALT+G is used for perspective view and G is used for GameView)

#Synced-CL 50006372

diff --git a/Engine/Source/Editor/AdvancedPreviewScene/Private/AdvancedPreviewSceneCommands.cpp b/Engine/Source/Editor/AdvancedPreviewScene/Private/AdvancedPreviewSceneCommands.cpp
--- a/Engine/Source/Editor/AdvancedPreviewScene/Private/AdvancedPreviewSceneCommands.cpp
+++ b/Engine/Source/Editor/AdvancedPreviewScene/Private/AdvancedPreviewSceneCommands.cpp
@@ -12,7 +12,7 @@ void FAdvancedPreviewSceneCommands::RegisterCommands()
 {
 	UI_COMMAND(ToggleEnvironment, "Toggle Environment", "Toggles Environment visibility", EUserInterfaceActionType::ToggleButton, FInputChord(EKeys::I));
 	UI_COMMAND(ToggleFloor, "Toggle Floor", "Toggles floor visibility", EUserInterfaceActionType::ToggleButton, FInputChord(EKeys::O));
-	UI_COMMAND(ToggleGrid, "Toggle Grid", "Toggles grid visibility", EUserInterfaceActionType::ToggleButton, FInputChord());
+	UI_COMMAND(ToggleGrid, "Toggle Grid", "Toggles grid visibility", EUserInterfaceActionType::ToggleButton, FInputChord(EModifierKey::Alt | EModifierKey::Shift, EKeys::G));
 	UI_COMMAND(TogglePostProcessing, "Toggle Post Processing", "Toggles whether Post Processing is enabled", EUserInterfaceActionType::ToggleButton, FInputChord(EKeys::P));
 }
 
diff --git a/Engine/Source/Editor/StaticMeshEditor/Private/StaticMeshEditorActions.cpp b/Engine/Source/Editor/StaticMeshEditor/Private/StaticMeshEditorActions.cpp
--- a/Engine/Source/Editor/StaticMeshEditor/Private/StaticMeshEditorActions.cpp
+++ b/Engine/Source/Editor/StaticMeshEditor/Private/StaticMeshEditorActions.cpp
@@ -18,7 +18,7 @@ void FStaticMeshEditorCommands::RegisterCommands()
 	UI_COMMAND(SetShowVertexColor, "Vert Colors", "Toggles vertex colors.", EUserInterfaceActionType::ToggleButton, FInputChord());
 	UI_COMMAND(SetShowPhysicalMaterialMasks, "Physical Material Masks", "Toggles physical material masks.", EUserInterfaceActionType::ToggleButton, FInputChord());
 	UI_COMMAND(SetDrawUVs, "UV", "Toggles display of the static mesh's UVs for the specified channel.", EUserInterfaceActionType::ToggleButton, FInputChord());
-	UI_COMMAND(SetShowGrid, "Grid", "Displays the viewport grid.", EUserInterfaceActionType::ToggleButton, FInputChord());
+	UI_COMMAND(SetShowGrid, "Grid", "Displays the viewport grid.", EUserInterfaceActionType::ToggleButton, FInputChord(EKeys::G, EModifierKey::Alt | EModifierKey::Shift));
 	UI_COMMAND(SetShowBounds, "Bounds", "Toggles display of the bounds of the static mesh.", EUserInterfaceActionType::ToggleButton, FInputChord());
 	UI_COMMAND(SetShowSimpleCollision, "Simple Collision", "Toggles display of the simplified collision mesh of the static mesh, if present.", EUserInterfaceActionType::ToggleButton, FInputChord());
 	UI_COMMAND(SetShowComplexCollision, "Complex Collision", "Toggles display of the complex collision for this static mesh.", EUserInterfaceActionType::ToggleButton, FInputChord());
diff --git a/Engine/Source/Editor/UnrealEd/Private/EditorShowFlags.cpp b/Engine/Source/Editor/UnrealEd/Private/EditorShowFlags.cpp
--- a/Engine/Source/Editor/UnrealEd/Private/EditorShowFlags.cpp
+++ b/Engine/Source/Editor/UnrealEd/Private/EditorShowFlags.cpp
@@ -62,6 +62,7 @@ TArray<FShowFlagData>& GetShowFlagMenuItems()
 			EngineFlagsChords.Add("Landscape", FInputChord(EKeys::L, EModifierKey::Alt) );
 			EngineFlagsChords.Add("Volumes", FInputChord(EKeys::O, EModifierKey::Alt) );
 			EngineFlagsChords.Add("Cameras", FInputChord() );
+			EngineFlagsChords.Add("Grid", FInputChord(EKeys::G, EModifierKey::Alt | EModifierKey::Shift) );
 
 			struct FIterSink
 			{
