[Static Mesh Editor] Enable the new gizmo & viewport interactions for the static mesh editor.

#rnx
#Synced-CL 51471837

diff --git a/Engine/Source/Editor/StaticMeshEditor/Private/StaticMeshEditor.cpp b/Engine/Source/Editor/StaticMeshEditor/Private/StaticMeshEditor.cpp
--- a/Engine/Source/Editor/StaticMeshEditor/Private/StaticMeshEditor.cpp
+++ b/Engine/Source/Editor/StaticMeshEditor/Private/StaticMeshEditor.cpp
@@ -65,6 +65,7 @@
 #include "StaticMeshEditorViewportClient.h"
 #include "AdvancedPreviewScene.h"
 #include "SnappingUtils.h"
+#include "Tools/EdModeInteractiveToolsContext.h"
 
 #define LOCTEXT_NAMESPACE "StaticMeshEditor"
 
@@ -368,6 +369,11 @@ void FStaticMeshEditor::InitStaticMeshEditor( const EToolkitMode::Type Mode, con
 	const bool bCreateDefaultToolbar = true;
 	FAssetEditorToolkit::InitAssetEditor( Mode, InitToolkitHost, StaticMeshEditorAppIdentifier, StandaloneDefaultLayout, bCreateDefaultToolbar, bCreateDefaultStandaloneMenu, ObjectToEdit );
 
+	if (UModeManagerInteractiveToolsContext* ToolsContext = GetEditorModeManager().GetInteractiveToolsContext())
+	{
+		ToolsContext->SetViewportInteractionsEnabled(true);
+	}
+
 	StaticMeshDetailsView->SetIsPropertyEditingEnabledDelegate(FIsPropertyEditingEnabled::CreateLambda([this]
 	{
 		return GetOpenMethod() == EAssetOpenMethod::Edit;
@@ -393,12 +399,6 @@ void FStaticMeshEditor::PostInitAssetEditor()
 	{
 		FStaticMeshEditorViewportClient& ViewportClient = StaticMeshViewport->GetViewportClient();
 		ViewportClient.ReceivedFocus(ViewportClient.Viewport);
-
-		if (FEditorModeTools* const ModeTools = ViewportClient.GetModeTools())
-		{
-			// Opt out of new ITF gizmos and viewport inputs for now
-			ModeTools->SetSupportsViewportITF(false);
-		}
 	}
 
 	// Static mesh editor code generally assumes the SocketManager exists, so make sure it does (in case the tab manager / sockets window hasn't already done so)
diff --git a/Engine/Source/Editor/StaticMeshEditor/Private/StaticMeshEditorViewportClient.cpp b/Engine/Source/Editor/StaticMeshEditor/Private/StaticMeshEditorViewportClient.cpp
--- a/Engine/Source/Editor/StaticMeshEditor/Private/StaticMeshEditorViewportClient.cpp
+++ b/Engine/Source/Editor/StaticMeshEditor/Private/StaticMeshEditorViewportClient.cpp
@@ -32,6 +32,7 @@
 #include "Editor/EditorPerProjectUserSettings.h"
 #include "StaticMeshEditorTools.h"
 #include "AssetViewerSettings.h"
+#include "GizmoEdModeInterface.h"
 #include "UnrealWidget.h"
 
 #include "Rendering/NaniteResources.h"
@@ -249,73 +250,11 @@ void FStaticMeshEditorViewportClient::TrackingStarted( const struct FInputEventS
 
 	if( !bManipulating && bIsDraggingWidget && !bTrackingHandledExternally)
 	{
-		Widget->SetSnapEnabled(true);
-		const TArray<UStaticMeshSocket*> SelectedSockets = StaticMeshEditorPtr.Pin()->GetSelectedSockets();
-		if (SelectedSockets.Num())
-		{
-			FText TransText;
-			if( GetWidgetMode() == UE::Widget::WM_Rotate )
-			{
-				TransText = LOCTEXT("FStaticMeshEditorViewportClient_RotateSocket", "Rotate Socket");
-			}
-			else if (GetWidgetMode() == UE::Widget::WM_Translate)
-			{
-				if( InInputState.IsLeftMouseButtonPressed() && (Widget->GetCurrentAxis() & EAxisList::XYZ) )
-				{
-					const bool bAltDown = InInputState.IsAltButtonPressed();
-					if ( bAltDown )
-					{
-						// Rather than moving/rotating the selected socket, copy it and move the copy instead
-						StaticMeshEditorPtr.Pin()->DuplicateSelectedSockets();
-					}
-				}
-
-				TransText = LOCTEXT("FStaticMeshEditorViewportClient_TranslateSocket", "Translate Socket");
-			}
-
-			if (!TransText.IsEmpty())
-			{
-				GEditor->BeginTransaction(TransText);
-			}
-		}
+		const bool bDuplicate = InInputState.IsLeftMouseButtonPressed()
+			&& InInputState.IsAltButtonPressed()
+			&& (Widget->GetCurrentAxis() & EAxisList::XYZ);
 		
-		const bool bSelectedPrim = StaticMeshEditorPtr.Pin()->HasSelectedPrims();
-		if (bSelectedPrim)
-		{
-			FText TransText;
-			if (GetWidgetMode() == UE::Widget::WM_Rotate)
-			{
-				TransText = LOCTEXT("FStaticMeshEditorViewportClient_RotateCollision", "Rotate Collision");
-			}
-			else if (GetWidgetMode() == UE::Widget::WM_Scale)
-			{
-				TransText = LOCTEXT("FStaticMeshEditorViewportClient_ScaleCollision", "Scale Collision");
-			}
-			else if (GetWidgetMode() == UE::Widget::WM_Translate)
-			{
-				if (InInputState.IsLeftMouseButtonPressed() && (Widget->GetCurrentAxis() & EAxisList::XYZ))
-				{
-					const bool bAltDown = InInputState.IsAltButtonPressed();
-					if (bAltDown)
-					{
-						// Rather than moving/rotating the selected primitives, copy them and move the copies instead
-						StaticMeshEditorPtr.Pin()->DuplicateSelectedPrims(NULL);
-					}
-				}
-
-				TransText = LOCTEXT("FStaticMeshEditorViewportClient_TranslateCollision", "Translate Collision");
-			}
-			if (!TransText.IsEmpty())
-			{
-				GEditor->BeginTransaction(TransText);
-				if (StaticMesh->GetBodySetup())
-				{
-					StaticMesh->GetBodySetup()->Modify();
-				}
-			}
-		}
-
-		bManipulating = true;
+		BeginManipulation(bDuplicate);
 	}
 
 }
@@ -397,6 +336,86 @@ void FStaticMeshEditorViewportClient::TrackingStopped()
 	}
 }
 
+bool FStaticMeshEditorViewportClient::BeginTransform(const FGizmoState& InState)
+{
+	BeginManipulation(IsAltDown(Viewport) && GetWidgetMode() == UE::Widget::WM_Translate);
+	return true;
+}
+
+void FStaticMeshEditorViewportClient::BeginManipulation(bool bDuplicate)
+{
+	Widget->SetSnapEnabled(true);
+	const TArray<UStaticMeshSocket*> SelectedSockets = StaticMeshEditorPtr.Pin()->GetSelectedSockets();
+	if (SelectedSockets.Num())
+	{
+		FText TransText;
+		if( GetWidgetMode() == UE::Widget::WM_Rotate )
+		{
+			TransText = LOCTEXT("FStaticMeshEditorViewportClient_RotateSocket", "Rotate Socket");
+		}
+		else if (GetWidgetMode() == UE::Widget::WM_Translate)
+		{
+			if (bDuplicate)
+			{
+				// Rather than moving/rotating the selected socket, copy it and move the copy instead
+				StaticMeshEditorPtr.Pin()->DuplicateSelectedSockets();
+			}
+
+			TransText = LOCTEXT("FStaticMeshEditorViewportClient_TranslateSocket", "Translate Socket");
+		}
+
+		if (!TransText.IsEmpty())
+		{
+			GEditor->BeginTransaction(TransText);
+		}
+	}
+	
+	const bool bSelectedPrim = StaticMeshEditorPtr.Pin()->HasSelectedPrims();
+	if (bSelectedPrim)
+	{
+		FText TransText;
+		if (GetWidgetMode() == UE::Widget::WM_Rotate)
+		{
+			TransText = LOCTEXT("FStaticMeshEditorViewportClient_RotateCollision", "Rotate Collision");
+		}
+		else if (GetWidgetMode() == UE::Widget::WM_Scale)
+		{
+			TransText = LOCTEXT("FStaticMeshEditorViewportClient_ScaleCollision", "Scale Collision");
+		}
+		else if (GetWidgetMode() == UE::Widget::WM_Translate)
+		{
+			if (bDuplicate)
+			{
+				// Rather than moving/rotating the selected primitives, copy them and move the copies instead
+				StaticMeshEditorPtr.Pin()->DuplicateSelectedPrims(NULL);
+			}
+
+			TransText = LOCTEXT("FStaticMeshEditorViewportClient_TranslateCollision", "Translate Collision");
+		}
+		if (!TransText.IsEmpty())
+		{
+			GEditor->BeginTransaction(TransText);
+			if (StaticMesh->GetBodySetup())
+			{
+				StaticMesh->GetBodySetup()->Modify();
+			}
+		}
+	}
+
+	bManipulating = true;
+}
+
+bool FStaticMeshEditorViewportClient::EndTransform(const FGizmoState& InState)
+{
+	if (bManipulating)
+	{
+		bManipulating = false;
+		GEditor->EndTransaction();
+		return true;
+	}
+	return false;
+}
+
 FVector FStaticMeshEditorViewportClient::GetWidgetLocation() const
 {
 	if (IsCustomModeUsingWidget())
diff --git a/Engine/Source/Editor/StaticMeshEditor/Private/StaticMeshEditorViewportClient.h b/Engine/Source/Editor/StaticMeshEditor/Private/StaticMeshEditorViewportClient.h
--- a/Engine/Source/Editor/StaticMeshEditor/Private/StaticMeshEditorViewportClient.h
+++ b/Engine/Source/Editor/StaticMeshEditor/Private/StaticMeshEditorViewportClient.h
@@ -43,6 +43,9 @@ public:
 	virtual FMatrix GetWidgetCoordSystem() const override;
 	virtual ECoordSystem GetWidgetCoordSystemSpace() const override;
 	virtual bool ShouldOrbitCamera() const override;
+	
+	virtual bool BeginTransform(const FGizmoState& InState) override;
+	virtual bool EndTransform(const FGizmoState& InState) override;
 
 	/** 
 	 *	Updates the static mesh and static mesh component being used in the Static Mesh Editor.
@@ -180,6 +183,8 @@ protected:
 private:
 	bool IsCustomModeUsingWidget() const;
 
+	void BeginManipulation(bool bDuplicate);
+
 	/** The Simplygon logo to be drawn when Simplygon has been used on the static mesh. */
 	UTexture2D* SimplygonLogo;
 
