[Static Mesh Editor] Fix alt+drag creating multiple duplicates, and broken Undo/Redo when using new Gizmo

#rnx
#Synced-CL 52630564

diff --git a/Engine/Source/Editor/StaticMeshEditor/Private/StaticMeshEditorEdMode.h b/Engine/Source/Editor/StaticMeshEditor/Private/StaticMeshEditorEdMode.h
--- a/Engine/Source/Editor/StaticMeshEditor/Private/StaticMeshEditorEdMode.h
+++ b/Engine/Source/Editor/StaticMeshEditor/Private/StaticMeshEditorEdMode.h
@@ -14,4 +14,7 @@ public:
 
 	virtual bool ShouldDrawWidget() const override { return true; }
 	virtual bool RequiresLegacyViewportInteractions() const override { return false; }
+
+	// Explicitly handle duplication. Reconsider when sockets have TypedElement handlers.
+	virtual EEditAction::Type GetActionEditDuplicate() override { return EEditAction::Process; }
 };
